Copying //UE4/Dev-Sequencer to //UE4/Main

diff --git a/Engine/Source/Editor/Sequencer/Private/SequencerTimeSliderController.h b/Engine/Source/Editor/Sequencer/Private/SequencerTimeSliderController.h
new file mode 100644
--- /dev/null
+++ b/Engine/Source/Editor/Sequencer/Private/SequencerTimeSliderController.h
@@ -0,0 +1,203 @@
+// Copyright 1998-2016 Epic Games, Inc. All Rights Reserved.
+
+#pragma once
+
+#include "Editor/SequencerWidgets/Public/ITimeSlider.h"
+
+struct FContextMenuSuppressor;
+
+struct FPaintPlaybackRangeArgs
+{
+	FPaintPlaybackRangeArgs()
+		: StartBrush(nullptr), EndBrush(nullptr), BrushWidth(0.f)
+	{}
+
+	FPaintPlaybackRangeArgs(const FSlateBrush* InStartBrush, const FSlateBrush* InEndBrush, float InBrushWidth)
+		: StartBrush(InStartBrush), EndBrush(InEndBrush), BrushWidth(InBrushWidth)
+	{}
+	/** Brush to use for the start bound */
+	const FSlateBrush* StartBrush;
+	/** Brush to use for the end bound */
+	const FSlateBrush* EndBrush;
+	/** The width of the above brushes, in slate units */
+	float BrushWidth;
+};
+
+struct FPaintSectionAreaViewArgs
+{
+	FPaintSectionAreaViewArgs()
+		: bDisplayTickLines(false), bDisplayScrubPosition(false)
+	{}
+
+	/** Whether to display tick lines */
+	bool bDisplayTickLines;
+	/** Whether to display the scrub position */
+	bool bDisplayScrubPosition;
+	/** Optional Paint args for the playback range*/
+	TOptional<FPaintPlaybackRangeArgs> PlaybackRangeArgs;
+};
+
+/**
+ * A time slider controller for sequencer
+ * Draws and manages time data for a Sequencer
+ */
+class FSequencerTimeSliderController : public ITimeSliderController
+{
+public:
+	FSequencerTimeSliderController( const FTimeSliderArgs& InArgs );
+
+	/** ITimeSliderController Interface */
+	virtual int32 OnPaintTimeSlider( bool bMirrorLabels, const FGeometry& AllottedGeometry, const FSlateRect& MyClippingRect, FSlateWindowElementList& OutDrawElements, int32 LayerId, const FWidgetStyle& InWidgetStyle, bool bParentEnabled ) const override;
+	virtual FReply OnMouseButtonDown( SWidget& WidgetOwner, const FGeometry& MyGeometry, const FPointerEvent& MouseEvent ) override;
+	virtual FReply OnMouseButtonUp( SWidget& WidgetOwner, const FGeometry& MyGeometry, const FPointerEvent& MouseEvent ) override;
+	virtual FReply OnMouseMove( SWidget& WidgetOwner, const FGeometry& MyGeometry, const FPointerEvent& MouseEvent ) override;
+	virtual FReply OnMouseWheel( SWidget& WidgetOwner, const FGeometry& MyGeometry, const FPointerEvent& MouseEvent ) override;
+	virtual FCursorReply OnCursorQuery( TSharedRef<const SWidget> WidgetOwner, const FGeometry& MyGeometry, const FPointerEvent& CursorEvent ) const override;
+	/** End ITimeSliderController Interface */
+
+	/** Get the current view range for this controller */
+	virtual FAnimatedRange GetViewRange() const override { return TimeSliderArgs.ViewRange.Get(); }
+
+	/** Get the current clamp range for this controller */
+	virtual FAnimatedRange GetClampRange() const override { return TimeSliderArgs.ClampRange.Get(); }
+
+	/** Convert time to frame */
+	virtual int32 TimeToFrame(float Time) const override;
+
+	/** Convert time to frame */
+	virtual float FrameToTime(int32 Frame) const override;
+
+	/**
+	 * Clamp the given range to the clamp range 
+	 *
+	 * @param NewRangeMin		The new lower bound of the range
+	 * @param NewRangeMax		The new upper bound of the range
+	 */	
+	void ClampViewRange(float& NewRangeMin, float& NewRangeMax);
+
+	/**
+	 * Set a new range based on a min, max and an interpolation mode
+	 * 
+	 * @param NewRangeMin		The new lower bound of the range
+	 * @param NewRangeMax		The new upper bound of the range
+	 * @param Interpolation		How to set the new range (either immediately, or animated)
+	 */
+	virtual void SetViewRange( float NewRangeMin, float NewRangeMax, EViewRangeInterpolation Interpolation ) override;
+
+	/**
+	 * Set a new clamp range based on a min, max
+	 * 
+	 * @param NewRangeMin		The new lower bound of the clamp range
+	 * @param NewRangeMax		The new upper bound of the clamp range
+	 */
+	virtual void SetClampRange( float NewRangeMin, float NewRangeMax ) override;
+
+	/**
+	 * Zoom the range by a given delta.
+	 * 
+	 * @param InDelta		The total amount to zoom by (+ve = zoom out, -ve = zoom in)
+	 * @param ZoomBias		Bias to apply to lower/upper extents of the range. (0 = lower, 0.5 = equal, 1 = upper)
+	 */
+	bool ZoomByDelta( float InDelta, float ZoomBias = 0.5f );
+
+	/**
+	 * Pan the range by a given delta
+	 * 
+	 * @param InDelta		The total amount to pan by (+ve = pan forwards in time, -ve = pan backwards in time)
+	 */
+	void PanByDelta( float InDelta );
+
+	/**
+	 * Draws major tick lines in the section view                                                              
+	 */
+	int32 OnPaintSectionView( const FGeometry& AllottedGeometry, const FSlateRect& MyClippingRect, FSlateWindowElementList& OutDrawElements, int32 LayerId, bool bEnabled, const FPaintSectionAreaViewArgs& Args ) const;
+
+private:
+	/**
+	 * Call this method when the user's interaction has changed the scrub position
+	 *
+	 * @param NewValue				Value resulting from the user's interaction
+	 * @param bIsScrubbing			True if done via scrubbing, false if just releasing scrubbing
+	 */
+	void CommitScrubPosition( float NewValue, bool bIsScrubbing );
+
+	/**
+	 * Draws time tick marks
+	 *
+	 * @param OutDrawElements	List to add draw elements to
+	 * @param RangeToScreen		Time range to screen space converter
+	 * @param InArgs			Parameters for drawing the tick lines
+	 */
+	void DrawTicks( FSlateWindowElementList& OutDrawElements, const struct FScrubRangeToScreen& RangeToScreen, struct FDrawTickArgs& InArgs ) const;
+
+	/**
+	 * Draws the playback range
+	 * @return the new layer ID
+	 */
+	int32 DrawPlaybackRange(const FGeometry& AllottedGeometry, const FSlateRect& MyClippingRect, FSlateWindowElementList& OutDrawElements, int32 LayerId, const FScrubRangeToScreen& RangeToScreen, const FPaintPlaybackRangeArgs& Args) const;
+
+private:
+	/**
+	 * Hit test the lower bound of the playback range
+	 */
+	bool HitTestPlaybackStart(const FScrubRangeToScreen& RangeToScreen, const TRange<float>& PlaybackRange, float LocalHitPositionX, float ScrubPosition) const;
+
+	/**
+	 * Hit test the upper bound of the playback range
+	 */
+	bool HitTestPlaybackEnd(const FScrubRangeToScreen& RangeToScreen, const TRange<float>& PlaybackRange, float LocalHitPositionX, float ScrubPosition) const;
+
+	void SetPlaybackRangeStart(float NewStart);
+	void SetPlaybackRangeEnd(float NewEnd);
+
+	TSharedRef<SWidget> OpenSetPlaybackRangeMenu(float MouseTime);
+
+private:
+	FTimeSliderArgs TimeSliderArgs;
+	/** The size of the scrub handle */
+	float ScrubHandleSize;
+	/** Brush for drawing an upwards facing scrub handle */
+	const FSlateBrush* ScrubHandleUp;
+	/** Brush for drawing a downwards facing scrub handle */
+	const FSlateBrush* ScrubHandleDown;
+	/** Total mouse delta during dragging **/
+	float DistanceDragged;
+	/** If we are dragging the scrubber or dragging to set the time range */
+	enum DragType
+	{
+		DRAG_SCRUBBING_TIME,
+		DRAG_SETTING_RANGE,
+		DRAG_START_RANGE,
+		DRAG_END_RANGE,
+		DRAG_NONE
+	};
+	DragType MouseDragType;
+	/** If we are currently panning the panel */
+	bool bPanning;
+	/** Mouse down time range */
+	FVector2D MouseDownRange;
+	/** Range stack */
+	TArray<FVector2D> RangeStack;
+	/** When > 0, we should not show context menus */
+	int32 ContextMenuSuppression;
+	friend FContextMenuSuppressor;
+};
+
+struct FContextMenuSuppressor
+{
+	FContextMenuSuppressor(TSharedRef<FSequencerTimeSliderController> InTimeSliderController)
+		: TimeSliderController(InTimeSliderController)
+	{
+		++TimeSliderController->ContextMenuSuppression;
+	}
+	~FContextMenuSuppressor()
+	{
+		--TimeSliderController->ContextMenuSuppression;
+	}
+
+private:
+	FContextMenuSuppressor(const FContextMenuSuppressor&);
+	FContextMenuSuppressor& operator=(const FContextMenuSuppressor&);
+
+	TSharedRef<FSequencerTimeSliderController> TimeSliderController;
+};
\ No newline at end of file
diff --git a/Engine/Source/Editor/Sequencer/Public/IKeyArea.h b/Engine/Source/Editor/Sequencer/Public/IKeyArea.h
--- a/Engine/Source/Editor/Sequencer/Public/IKeyArea.h
+++ b/Engine/Source/Editor/Sequencer/Public/IKeyArea.h
@@ -194,7 +194,7 @@ public:
 	 *
 	 * @return Key area color.
 	 */
-	virtual FLinearColor GetColor() = 0;
+	virtual TOptional<FLinearColor> GetColor() = 0;
 
 	/**
 	 * Get the data structure representing the specified key.
diff --git a/Engine/Source/Editor/Sequencer/Public/ISequencer.h b/Engine/Source/Editor/Sequencer/Public/ISequencer.h
--- a/Engine/Source/Editor/Sequencer/Public/ISequencer.h
+++ b/Engine/Source/Editor/Sequencer/Public/ISequencer.h
@@ -14,6 +14,7 @@ class UMovieSceneSequence;
 class UAnimSequence;
 class UMovieScene;
 class UMovieSceneSection;
+class UMovieSceneSubSection;
 enum class EMapChangeType : uint8;
 
 
@@ -70,9 +71,9 @@ public:
 	/**
 	 * Focuses a sub-movie scene (MovieScene within a MovieScene) in the sequencer.
 	 * 
-	 * @param SequenceInstance	Sub-sequence instance to focus.
+	 * @param Section The sub-movie scene section containing the sequence instance to get.
 	 */
-	virtual void FocusSequenceInstance(TSharedRef<FMovieSceneSequenceInstance> SequenceInstance) = 0;
+	virtual void FocusSequenceInstance(UMovieSceneSubSection& Section) = 0;
 	
 	/**
 	 * Create a new binding for the specified object
@@ -150,6 +151,9 @@ public:
 	 */
 	virtual void SetGlobalTime(float Time, const bool& bAllowSnappingToFrames = false) = 0;
 
+	/** Set the global time directly, without performing any auto-scroll */
+	virtual void SetGlobalTimeDirectly(float Time, const bool& bAllowSnappingToFrames = false) = 0;
+
 	/** @return The current view range */
 	virtual FAnimatedRange GetViewRange() const
 	{
diff --git a/Engine/Source/Editor/Sequencer/Public/ISequencerSection.h b/Engine/Source/Editor/Sequencer/Public/ISequencerSection.h
--- a/Engine/Source/Editor/Sequencer/Public/ISequencerSection.h
+++ b/Engine/Source/Editor/Sequencer/Public/ISequencerSection.h
@@ -2,6 +2,8 @@
 
 #pragma once
 
+#include "SequencerSectionPainter.h"
+
 class IKeyArea;
 
 namespace SequencerSectionConstants
@@ -19,10 +21,6 @@ namespace SequencerSectionConstants
 	const FName SelectionColorName("SelectionColor");
 
 	const FName SelectionInactiveColorName("SelectionColorInactive");
-	
-	const FName DefaultSectionGripLeftImageName("Sequencer.DefaultSectionGripLeft");
-	
-	const FName DefaultSectionGripRightImageName("Sequencer.DefaultSectionGripRight");
 }
 
 /**
@@ -48,14 +46,10 @@ public:
 	/**
 	 * Called when the section should be painted
 	 *
-	 * @param AllottedGeometry		The geometry of the section
-	 * @param SectionClippingRect	The clipping rect around the section
-	 * @param OutDrawElements		The element list ot add elements to
-	 * @param LayerId				The ID of the starting layer to draw elements to
-	 * @param bParentEnabled		Whether or not the parent widget is enabled
-	 * @return						The new LayerId
+	 * @param Painter		Structure that affords common painting operations
+	 * @return				The new LayerId
 	 */
-	virtual int32 OnPaintSection( const FGeometry& AllottedGeometry, const FSlateRect& SectionClippingRect, FSlateWindowElementList& OutDrawElements, int32 LayerId, bool bParentEnabled ) const = 0;
+	virtual int32 OnPaintSection( FSequencerSectionPainter& InPainter ) const = 0;
 
 	/** Allows a section to override the brush to use for a key by handle.
 	 *
@@ -64,11 +58,6 @@ public:
 	 */
 	virtual const FSlateBrush* GetKeyBrush(FKeyHandle KeyHandle) const { return nullptr; }
 
-	/*
-	 * @return Whether to draw a background for the key area
-	 */
-	virtual bool ShouldDrawKeyAreaBackground() const { return true; }
-
 	/**
 	 * Called when the section is double clicked
 	 *
@@ -88,6 +77,11 @@ public:
 	 */
 	virtual FText GetSectionTitle() const = 0;
 
+	/**
+	 * @return The amount of padding to apply to non-interactive portions of the section interface (such as section text)
+	 */
+	virtual FMargin GetContentPadding() const { return FMargin(11.f, 6.f); }
+
 	/**
 	 * Generates the inner layout for this section
 	 *
@@ -102,10 +96,6 @@ public:
 	
 	virtual float GetSectionGripSize() const { return SequencerSectionConstants::DefaultSectionGripSize; }
 
-	virtual FName GetSectionGripLeftBrushName() const { return SequencerSectionConstants::DefaultSectionGripLeftImageName; }
-
-	virtual FName GetSectionGripRightBrushName() const { return SequencerSectionConstants::DefaultSectionGripRightImageName; }
-
 	/**
 	 * @return Whether or not the user can resize this section.
 	 */
@@ -147,4 +137,4 @@ public:
 	* @returns Whether or not the key area was deleted.
 	*/
 	virtual bool RequestDeleteKeyArea( const TArray<FName>& KeyAreaNamePath ) { return false; }
-};
+};
\ No newline at end of file
diff --git a/Engine/Source/Editor/Sequencer/Public/SequencerCommands.h b/Engine/Source/Editor/Sequencer/Public/SequencerCommands.h
--- a/Engine/Source/Editor/Sequencer/Public/SequencerCommands.h
+++ b/Engine/Source/Editor/Sequencer/Public/SequencerCommands.h
@@ -153,8 +153,8 @@ public:
 	/** Finds the viewed sequence asset in the content browser. */
 	TSharedPtr< FUICommandInfo > FindInContentBrowser;
 
-	/** Toggles whether the details view is enabled in the level editor. */
-	TSharedPtr< FUICommandInfo > ToggleDetailsView;
+	/** Toggles whether to show combined keys at the top node level. */
+	TSharedPtr< FUICommandInfo > ToggleCombinedKeyframes;
 
 	/** Toggles whether the label browser is enabled in the level editor. */
 	TSharedPtr< FUICommandInfo > ToggleLabelBrowser;
@@ -165,6 +165,9 @@ public:
 	/** Toggles whether or not the curve editor should be shown. */
 	TSharedPtr< FUICommandInfo > ToggleShowCurveEditor;
 
+	/** Toggles whether or not the curve editor time range should be linked to the sequencer. */
+	TSharedPtr< FUICommandInfo > ToggleLinkCurveEditorTimeRange;
+
 	/** Enable the move tool */
 	TSharedPtr< FUICommandInfo > MoveTool;
 
@@ -180,6 +183,9 @@ public:
 	/** Convert the selected possessed objects to spawnables. These will be spawned and destroyed by sequencer as per object's the spawn track. */
 	TSharedPtr< FUICommandInfo > ConvertToSpawnable;
 
+	/** Discard all changes to the current movie scene. */
+	TSharedPtr< FUICommandInfo > DiscardChanges;
+
 	/**
 	 * Initialize commands
 	 */
diff --git a/Engine/Source/Editor/Sequencer/Public/SequencerSectionPainter.h b/Engine/Source/Editor/Sequencer/Public/SequencerSectionPainter.h
new file mode 100644
--- /dev/null
+++ b/Engine/Source/Editor/Sequencer/Public/SequencerSectionPainter.h
@@ -0,0 +1,72 @@
+// Copyright 1998-2015 Epic Games, Inc. All Rights Reserved.
+
+#pragma once
+
+struct FTimeToPixel;
+class UMovieSceneTrack;
+class UMovieSceneSection;
+
+/** Structure representing a specific region on a section to tint */
+struct FSequencerSectionTintRegion
+{
+	FSequencerSectionTintRegion(const FColor& InColor)
+		: Color(InColor)
+	{}
+
+	FSequencerSectionTintRegion(const FColor& InColor, float InStartTime, float InEndTime)
+		: Color(InColor), Range(TRange<float>(InStartTime, InEndTime))
+	{}
+
+	/** The color of the tint */
+	FColor Color;
+
+	/** The time range to tint */
+	TOptional<TRange<float>> Range;
+};
+
+/** Class that wraps up common section painting functionality */
+class SEQUENCER_API FSequencerSectionPainter
+{
+public:
+	/** Constructor */
+	FSequencerSectionPainter(FSlateWindowElementList& OutDrawElements, const FGeometry& InSectionGeometry, UMovieSceneSection& Section);
+
+	/** Paint the section background with the specified tint override */
+	int32 PaintSectionBackground(const FColor& Tint);
+
+	/** Paint the section background with the tint stored on the track */
+	int32 PaintSectionBackground();
+
+	/** Get the track that this painter is painting sections for */
+	UMovieSceneTrack* GetTrack() const;
+
+public:
+
+	/** Paint the section background with the specified tints */
+	virtual int32 PaintSectionBackground(const FSequencerSectionTintRegion* Tints, uint32 NumTints) = 0;
+
+	/** Get a time-to-pixel converter for the section */
+	virtual const FTimeToPixel& GetTimeConverter() const = 0;
+
+public:
+
+	/** The movie scene section we're painting */
+	UMovieSceneSection& Section;
+
+	/** List of slate draw elements - publicly modifiable */
+	FSlateWindowElementList& DrawElements;
+
+	/** The full geometry of the section. This is the width of the track area in the case of infinite sections */
+	FGeometry SectionGeometry;
+	
+	/** The full clipping rectangle for the section */
+	FSlateRect SectionClippingRect;
+	
+	/** The layer ID we're painting on */
+	int32 LayerId;
+
+	/** Whether our parent widget is enabled or not */
+	bool bParentEnabled;
+	
+};
+
diff --git a/Engine/Source/Editor/SequencerWidgets/Private/SSequencerTimeSlider.cpp b/Engine/Source/Editor/SequencerWidgets/Private/SSequencerTimeSlider.cpp
new file mode 100644
--- /dev/null
+++ b/Engine/Source/Editor/SequencerWidgets/Private/SSequencerTimeSlider.cpp
@@ -0,0 +1,52 @@
+// Copyright 1998-2016 Epic Games, Inc. All Rights Reserved.
+
+#include "SequencerWidgetsPrivatePCH.h"
+#include "SSequencerTimeSlider.h"
+
+#define LOCTEXT_NAMESPACE "STimeSlider"
+
+
+void SSequencerTimeSlider::Construct( const SSequencerTimeSlider::FArguments& InArgs, TSharedRef<ITimeSliderController> InTimeSliderController )
+{
+	TimeSliderController = InTimeSliderController;
+	bMirrorLabels = InArgs._MirrorLabels;
+}
+
+int32 SSequencerTimeSlider::OnPaint( const FPaintArgs& Args, const FGeometry& AllottedGeometry, const FSlateRect& MyClippingRect, FSlateWindowElementList& OutDrawElements, int32 LayerId, const FWidgetStyle& InWidgetStyle, bool bParentEnabled ) const
+{
+	int32 NewLayer = TimeSliderController->OnPaintTimeSlider( bMirrorLabels, AllottedGeometry, MyClippingRect, OutDrawElements, LayerId, InWidgetStyle, bParentEnabled );
+
+	return FMath::Max( NewLayer, SCompoundWidget::OnPaint(Args, AllottedGeometry, MyClippingRect, OutDrawElements, NewLayer, InWidgetStyle, ShouldBeEnabled( bParentEnabled ) ) );
+}
+
+FReply SSequencerTimeSlider::OnMouseButtonDown( const FGeometry& MyGeometry, const FPointerEvent& MouseEvent )
+{
+	return TimeSliderController->OnMouseButtonDown( *this, MyGeometry, MouseEvent );
+}
+
+FReply SSequencerTimeSlider::OnMouseButtonUp( const FGeometry& MyGeometry, const FPointerEvent& MouseEvent )
+{
+	return TimeSliderController->OnMouseButtonUp( *this,  MyGeometry, MouseEvent );
+}
+
+FReply SSequencerTimeSlider::OnMouseMove( const FGeometry& MyGeometry, const FPointerEvent& MouseEvent )
+{
+	return TimeSliderController->OnMouseMove( *this, MyGeometry, MouseEvent );
+}
+
+FVector2D SSequencerTimeSlider::ComputeDesiredSize( float ) const
+{
+	return FVector2D(100, 22);
+}
+
+FReply SSequencerTimeSlider::OnMouseWheel( const FGeometry& MyGeometry, const FPointerEvent& MouseEvent )
+{
+	return TimeSliderController->OnMouseWheel( *this, MyGeometry, MouseEvent );
+}
+
+FCursorReply SSequencerTimeSlider::OnCursorQuery(const FGeometry& MyGeometry, const FPointerEvent& CursorEvent) const
+{
+	return TimeSliderController->OnCursorQuery( SharedThis(this), MyGeometry, CursorEvent );
+}
+
+#undef LOCTEXT_NAMESPACE
diff --git a/Engine/Source/Editor/SequencerWidgets/Private/SequencerWidgetsModule.cpp b/Engine/Source/Editor/SequencerWidgets/Private/SequencerWidgetsModule.cpp
--- a/Engine/Source/Editor/SequencerWidgets/Private/SequencerWidgetsModule.cpp
+++ b/Engine/Source/Editor/SequencerWidgets/Private/SequencerWidgetsModule.cpp
@@ -1,44 +1,48 @@
 // Copyright 1998-2016 Epic Games, Inc. All Rights Reserved.
 
 #include "SequencerWidgetsPrivatePCH.h"
-#include "STimeSlider.h"
-#include "STimeRange.h"
+#include "ISequencerWidgetsModule.h"
 #include "ModuleManager.h"
+#include "SSequencerTimeSlider.h"
+#include "STimeRange.h"
 
 
-IMPLEMENT_MODULE( FSequencerWidgetsModule, SequencerWidgets );
-
-void FSequencerWidgetsModule::StartupModule()
+/**
+ * The public interface of SequencerModule
+ */
+class FSequencerWidgetsModule
+	: public ISequencerWidgetsModule
 {
+public:
 
-}
+	// ISequencerWidgetsModule interface
 
-void FSequencerWidgetsModule::ShutdownModule()
-{
+	TSharedRef<ITimeSlider> CreateTimeSlider(const TSharedRef<ITimeSliderController>& InController, bool bMirrorLabels) override
+	{
+		return SNew(SSequencerTimeSlider, InController)
+			.MirrorLabels(bMirrorLabels);
+	}
 
-}
+	TSharedRef<ITimeSlider> CreateTimeSlider(const TSharedRef<ITimeSliderController>& InController, const TAttribute<EVisibility>& VisibilityDelegate, bool bMirrorLabels) override
+	{
+		return SNew(SSequencerTimeSlider, InController)
+			.Visibility(VisibilityDelegate)
+			.MirrorLabels(bMirrorLabels);
+	}
 
-TSharedRef<ITimeSlider> FSequencerWidgetsModule::CreateTimeSlider( const TSharedRef<ITimeSliderController>& InController, bool bMirrorLabels )
-{
-	return
-		SNew( STimeSlider, InController )
-		.MirrorLabels( bMirrorLabels );
-}
+	TSharedRef<ITimeSlider> CreateTimeRange(const TSharedRef<ITimeSliderController>& InController, const TAttribute<EVisibility>& VisibilityDelegate, const TAttribute<bool>& ShowFrameNumbersDelegate, const TAttribute<float>& TimeSnapIntervalDelegate) override
+	{
+		return SNew(STimeRange, InController)
+			.Visibility(VisibilityDelegate)
+			.ShowFrameNumbers(ShowFrameNumbersDelegate)
+			.TimeSnapInterval(TimeSnapIntervalDelegate);
+	}
 
-TSharedRef<ITimeSlider> FSequencerWidgetsModule::CreateTimeSlider( const TSharedRef<ITimeSliderController>& InController, const TAttribute<EVisibility>& VisibilityDelegate, bool bMirrorLabels )
-{
-	return
-		SNew( STimeSlider, InController )
-		.Visibility(VisibilityDelegate)
-		.MirrorLabels( bMirrorLabels );
-}
+public:
+
+	virtual void StartupModule() override { }
+	virtual void ShutdownModule() override { }
+};
 
-TSharedRef<ITimeSlider> FSequencerWidgetsModule::CreateTimeRange( const TSharedRef<ITimeSliderController>& InController, const TAttribute<EVisibility>& VisibilityDelegate, const TAttribute<bool>& ShowFrameNumbersDelegate, const TAttribute<float>& TimeSnapIntervalDelegate )
-{
-	return 
-		SNew( STimeRange, InController)
-		.Visibility(VisibilityDelegate)
-		.ShowFrameNumbers(ShowFrameNumbersDelegate)
-		.TimeSnapInterval(TimeSnapIntervalDelegate);
-}
 
+IMPLEMENT_MODULE(FSequencerWidgetsModule, SequencerWidgets);
diff --git a/Engine/Source/Editor/SequencerWidgets/Private/SequencerWidgetsPrivatePCH.h b/Engine/Source/Editor/SequencerWidgets/Private/SequencerWidgetsPrivatePCH.h
--- a/Engine/Source/Editor/SequencerWidgets/Private/SequencerWidgetsPrivatePCH.h
+++ b/Engine/Source/Editor/SequencerWidgets/Private/SequencerWidgetsPrivatePCH.h
@@ -3,6 +3,4 @@
 #pragma once
 
 #include "SlateBasics.h"
-
-#include "SequencerWidgetsModule.h"
 #include "ITimeSlider.h"
diff --git a/Engine/Source/Editor/SequencerWidgets/Public/ISequencerWidgetsModule.h b/Engine/Source/Editor/SequencerWidgets/Public/ISequencerWidgetsModule.h
new file mode 100644
--- /dev/null
+++ b/Engine/Source/Editor/SequencerWidgets/Public/ISequencerWidgetsModule.h
@@ -0,0 +1,22 @@
+// Copyright 1998-2016 Epic Games, Inc. All Rights Reserved.
+
+#pragma once
+
+#include "ModuleInterface.h"
+
+
+class ITimeSlider;
+
+
+/**
+ * The public interface of SequencerModule
+ */
+class ISequencerWidgetsModule
+	: public IModuleInterface
+{
+public:
+
+	virtual TSharedRef<ITimeSlider> CreateTimeSlider(const TSharedRef<class ITimeSliderController>& InController, bool bMirrorLabels) = 0;
+	virtual TSharedRef<ITimeSlider> CreateTimeSlider(const TSharedRef<class ITimeSliderController>& InController, const TAttribute<EVisibility>& VisibilityDelegate, bool bMirrorLabels) = 0;
+	virtual TSharedRef<ITimeSlider> CreateTimeRange(const TSharedRef<class ITimeSliderController>& InController, const TAttribute<EVisibility>& VisibilityDelegate, const TAttribute<bool>& ShowFrameNumbersDelegate, const TAttribute<float>& TimeSnapIntervalDelegate) = 0;
+};
diff --git a/Engine/Source/Editor/UMGEditor/Private/Templates/WidgetTemplateClass.cpp b/Engine/Source/Editor/UMGEditor/Private/Templates/WidgetTemplateClass.cpp
--- a/Engine/Source/Editor/UMGEditor/Private/Templates/WidgetTemplateClass.cpp
+++ b/Engine/Source/Editor/UMGEditor/Private/Templates/WidgetTemplateClass.cpp
@@ -45,8 +45,8 @@ const FSlateBrush* GetEditorIcon_Deprecated(UWidget* Widget);
 const FSlateBrush* FWidgetTemplateClass::GetIcon() const
 {
 	// @todo UMG: remove after 4.12
-	auto DefaultUserWidget = WidgetClass->GetDefaultObject<UUserWidget>();
-	return GetEditorIcon_Deprecated(DefaultUserWidget);
+	UWidget* DefaultWidget = WidgetClass->GetDefaultObject<UWidget>();
+	return GetEditorIcon_Deprecated(DefaultWidget);
 	// return FClassIconFinder::FindIconForClass(WidgetClass);
 }
 
diff --git a/Engine/Source/Editor/UnrealEd/Private/FileHelpers.cpp b/Engine/Source/Editor/UnrealEd/Private/FileHelpers.cpp
--- a/Engine/Source/Editor/UnrealEd/Private/FileHelpers.cpp
+++ b/Engine/Source/Editor/UnrealEd/Private/FileHelpers.cpp
@@ -33,9 +33,12 @@
 #include "Engine/LevelStreaming.h"
 #include "GameMapsSettings.h"
 #include "AutoSaveUtils.h"
+#include "AssetRegistryModule.h"
+
 
 DEFINE_LOG_CATEGORY_STATIC(LogFileHelpers, Log, All);
 
+
 //definition of flag used to do special work when we're attempting to load the "startup map"
 bool FEditorFileUtils::bIsLoadingDefaultStartupMap = false;
 bool FEditorFileUtils::bIsPromptingForCheckoutAndSave = false;
@@ -896,6 +899,72 @@ static bool IsWorldDirty()
 //
 ///////////////////////////////////////////////////////////////////////////////////////////////////////////////////////
 
+void FEditorFileUtils::SaveAssetsAs(const TArray<UObject*>& Assets, TArray<UObject*>& OutSavedAssets)
+{
+	for (UObject* Asset : Assets)
+	{
+		const FString OldPackageName = Asset->GetOutermost()->GetName();
+		const FString OldPackagePath = FPackageName::GetLongPackagePath(OldPackageName);
+		const FString OldAssetName = FPackageName::GetLongPackageAssetName(OldPackageName);
+
+		FString NewPackageName;
+
+		// get destination for asset
+		bool FilenameValid = false;
+
+		while (!FilenameValid)
+		{
+			if (!OpenSaveAsDialog(Asset->GetClass(), OldPackagePath, OldAssetName, NewPackageName))
+			{
+				break;
+			}
+
+			FText OutError;
+			FilenameValid = FEditorFileUtils::IsFilenameValidForSaving(NewPackageName, OutError);
+		}
+
+		// process asset
+		if (NewPackageName.IsEmpty())
+		{
+			OutSavedAssets.Add(Asset); // user cancelled
+		}
+		else if (NewPackageName != OldPackageName)
+		{
+			// duplicate asset at destination
+			const FString NewAssetName = FPackageName::GetLongPackageAssetName(NewPackageName);
+			UPackage* DuplicatedPackage = CreatePackage(nullptr, *NewPackageName);
+			UObject* DuplicatedAsset = StaticDuplicateObject(Asset, DuplicatedPackage, *NewAssetName);
+
+			if (DuplicatedAsset != nullptr)
+			{
+				DuplicatedAsset->MarkPackageDirty();
+				FAssetRegistryModule::AssetCreated(DuplicatedAsset);
+				OutSavedAssets.Add(DuplicatedAsset);
+			}
+			else
+			{
+				OutSavedAssets.Add(Asset); // error duplicating
+			}
+		}
+		else
+		{
+			// save existing asset
+			OutSavedAssets.Add(Asset);
+		}
+	}
+
+	// save packages
+	TArray<UPackage*> PackagesToSave;
+
+	for (UObject* Asset : OutSavedAssets)
+	{
+		PackagesToSave.Add(Asset->GetOutermost());
+	}
+
+	FEditorFileUtils::PromptForCheckoutAndSave(PackagesToSave, true, false);
+}
+
+
 /**
  * Does a saveAs for the specified level.
  *
diff --git a/Engine/Source/Editor/UnrealEd/Private/LevelEditorViewport.cpp b/Engine/Source/Editor/UnrealEd/Private/LevelEditorViewport.cpp
--- a/Engine/Source/Editor/UnrealEd/Private/LevelEditorViewport.cpp
+++ b/Engine/Source/Editor/UnrealEd/Private/LevelEditorViewport.cpp
@@ -2094,20 +2094,8 @@ void FLevelEditorViewportClient::UpdateViewForLockedActor()
 			if (bLockedCameraView)
 			{
 				// If this is a camera actor, then inherit some other settings
-				TArray<UCameraComponent*> CamComps;
-				Actor->GetComponents<UCameraComponent>(CamComps);
-
-				UCameraComponent* CameraComponent = nullptr;
-				for (UCameraComponent* Comp : CamComps)
-				{
-					if (Comp->bIsActive)
-					{
-						CameraComponent = Comp;
-						break;
-					}
-				}
-
-				if (CameraComponent != NULL)
+				UCameraComponent* const CameraComponent = GetCameraComponentForLockedActor(Actor);
+				if (CameraComponent != nullptr)
 				{
 					bUseControllingActorViewInfo = true;
 					CameraComponent->GetCameraView(0.0f, ControllingActorViewInfo);
@@ -2123,7 +2111,6 @@ void FLevelEditorViewportClient::UpdateViewForLockedActor()
 	}
 }
 
-
 /*namespace ViewportDeadZoneConstants
 {
 	enum
diff --git a/Engine/Source/Editor/UnrealEd/Private/RichCurveEditorCommands.cpp b/Engine/Source/Editor/UnrealEd/Private/RichCurveEditorCommands.cpp
--- a/Engine/Source/Editor/UnrealEd/Private/RichCurveEditorCommands.cpp
+++ b/Engine/Source/Editor/UnrealEd/Private/RichCurveEditorCommands.cpp
@@ -23,6 +23,9 @@ void FRichCurveEditorCommands::RegisterCommands()
 	UI_COMMAND(FlattenTangents, "Flatten", "Flatten tangents", EUserInterfaceActionType::Button, FInputChord());
 	UI_COMMAND(StraightenTangents, "Straighten", "Straighten tangents", EUserInterfaceActionType::Button, FInputChord());
 
+	UI_COMMAND(BakeCurve, "Bake", "Bake curve", EUserInterfaceActionType::Button, FInputChord());
+	UI_COMMAND(ReduceCurve, "Reduce", "Reduce curve", EUserInterfaceActionType::Button, FInputChord());
+
 	UI_COMMAND(SetPreInfinityExtrapCycle, "Cycle", "Set the pre-infinity curve extrapolation to cycle.", EUserInterfaceActionType::RadioButton, FInputChord());
 	UI_COMMAND(SetPreInfinityExtrapCycleWithOffset, "Cycle with Offset", "Set the pre-infinity curve extrapolation to cycle with offset.", EUserInterfaceActionType::RadioButton, FInputChord());
 	UI_COMMAND(SetPreInfinityExtrapOscillate, "Oscillate", "Set the pre-infinity curve extrapolation to oscillate.", EUserInterfaceActionType::RadioButton, FInputChord());
diff --git a/Engine/Source/Editor/UnrealEd/Private/SCurveEditor.cpp b/Engine/Source/Editor/UnrealEd/Private/SCurveEditor.cpp
--- a/Engine/Source/Editor/UnrealEd/Private/SCurveEditor.cpp
+++ b/Engine/Source/Editor/UnrealEd/Private/SCurveEditor.cpp
@@ -8,6 +8,7 @@
 #include "SColorGradientEditor.h"
 #include "GenericCommands.h"
 #include "SNumericEntryBox.h"
+#include "STextEntryPopup.h"
 #include "CurveEditorSettings.h"
 
 #define LOCTEXT_NAMESPACE "SCurveEditor"
@@ -86,6 +87,8 @@ void SCurveEditor::Construct(const FArguments& InArgs)
 
 	TransactionIndex = -1;
 
+	ReduceTolerance = 0.001;
+
 	Settings = GetMutableDefault<UCurveEditorSettings>();
 
 	Commands->MapAction(FGenericCommands::Get().Undo,
@@ -144,6 +147,13 @@ void SCurveEditor::Construct(const FArguments& InArgs)
 	Commands->MapAction(FRichCurveEditorCommands::Get().StraightenTangents,
 		FExecuteAction::CreateSP(this, &SCurveEditor::OnFlattenOrStraightenTangents, false));
 
+	// Bake and reduce
+	Commands->MapAction(FRichCurveEditorCommands::Get().BakeCurve,
+		FExecuteAction::CreateSP(this, &SCurveEditor::OnBakeCurve));
+
+	Commands->MapAction(FRichCurveEditorCommands::Get().ReduceCurve,
+		FExecuteAction::CreateSP(this, &SCurveEditor::OnReduceCurve));
+
 	// Pre infinity extrapolation
 	Commands->MapAction(FRichCurveEditorCommands::Get().SetPreInfinityExtrapCycle,
 		FExecuteAction::CreateSP(this, &SCurveEditor::OnSelectPreInfinityExtrap, RCCE_Cycle),
@@ -2998,6 +3008,151 @@ void SCurveEditor::OnFlattenOrStraightenTangents(bool bFlattenTangents)
 	}
 }
 
+void SCurveEditor::OnBakeCurve()
+{
+	float BakeSampleRate = InputSnap.IsSet() ? InputSnap.Get() : 0.05f;
+
+	// Display dialog and let user enter sample rate.
+	GenericTextEntryModeless(
+		NSLOCTEXT("CurveEditor.Popups", "BakeSampleRate", "Sample Rate"),
+		FText::AsNumber( BakeSampleRate ),
+		FOnTextCommitted::CreateSP(this, &SCurveEditor::OnBakeCurveSampleRateCommitted)
+		);
+}
+
+void SCurveEditor::OnBakeCurveSampleRateCommitted(const FText& InText, ETextCommit::Type CommitInfo)
+{
+	CloseEntryPopupMenu();
+	if (CommitInfo == ETextCommit::OnEnter)
+	{
+		double NewBakeSampleRate = FCString::Atod(*InText.ToString());
+		const bool bIsNumber = InText.IsNumeric(); 
+		if(!bIsNumber)
+			return;
+
+		float BakeSampleRate = (float)NewBakeSampleRate;
+
+		const FScopedTransaction Transaction(LOCTEXT("CurveEditor_BakeCurve", "Bake Curve"));
+		CurveOwner->ModifyOwner();
+
+		bool bAnyCurveViewModelsSelected = AnyCurveViewModelsSelected();
+
+		TArray<FRichCurveEditInfo> ChangedCurveEditInfos;
+
+		// If keys are selected, bake between them
+		TMap<FRichCurve*, TInterval<float> > CurveRangeMap;
+		for (auto SelectedKey : SelectedKeys)
+		{
+			float SelectedTime = SelectedKey.Curve->GetKey(SelectedKey.KeyHandle).Time;
+			if (CurveRangeMap.Find(SelectedKey.Curve) != nullptr)
+			{
+				CurveRangeMap[SelectedKey.Curve].Include(SelectedTime);
+			}
+			else
+			{
+				CurveRangeMap.Add(SelectedKey.Curve, TInterval<float>(SelectedTime, SelectedTime));
+			}
+		}
+
+		if (CurveRangeMap.Num())
+		{
+			for (auto CurveToBake : CurveRangeMap)
+			{
+				CurveToBake.Key->BakeCurve(BakeSampleRate, CurveToBake.Value.Min, CurveToBake.Value.Max);
+				ChangedCurveEditInfos.Add(GetViewModelForCurve(CurveToBake.Key)->CurveInfo);
+			}
+		}
+		else
+		{
+			for (auto CurveViewModel : CurveViewModels)
+			{
+				if (!bAnyCurveViewModelsSelected || CurveViewModel->bIsSelected)
+				{
+					CurveViewModel->CurveInfo.CurveToEdit->BakeCurve(BakeSampleRate);
+					ChangedCurveEditInfos.Add(CurveViewModel->CurveInfo);
+				}
+			}
+		}
+
+		if (ChangedCurveEditInfos.Num())
+		{
+			CurveOwner->OnCurveChanged(ChangedCurveEditInfos);
+		}
+	}
+}
+
+void SCurveEditor::OnReduceCurve()
+{
+	// Display dialog and let user enter tolerance.
+	GenericTextEntryModeless(
+		NSLOCTEXT("CurveEditor.Popups", "ReduceCurveTolerance", "Tolerance"),
+		FText::AsNumber( ReduceTolerance ),
+		FOnTextCommitted::CreateSP(this, &SCurveEditor::OnReduceCurveToleranceCommitted)
+		);
+}
+
+void SCurveEditor::OnReduceCurveToleranceCommitted(const FText& InText, ETextCommit::Type CommitInfo)
+{
+	CloseEntryPopupMenu();
+	if (CommitInfo == ETextCommit::OnEnter)
+	{
+		double NewTolerance = FCString::Atod(*InText.ToString());
+		const bool bIsNumber = InText.IsNumeric(); 
+		if(!bIsNumber)
+			return;
+
+		ReduceTolerance = (float)NewTolerance;
+
+		const FScopedTransaction Transaction(LOCTEXT("CurveEditor_ReduceCurve", "Reduce Curve"));
+		CurveOwner->ModifyOwner();
+
+		bool bAnyCurveViewModelsSelected = AnyCurveViewModelsSelected();
+
+		TArray<FRichCurveEditInfo> ChangedCurveEditInfos;
+
+		// If keys are selected, bake between them
+		TMap<FRichCurve*, TInterval<float> > CurveRangeMap;
+		for (auto SelectedKey : SelectedKeys)
+		{
+			float SelectedTime = SelectedKey.Curve->GetKey(SelectedKey.KeyHandle).Time;
+			if (CurveRangeMap.Find(SelectedKey.Curve) != nullptr)
+			{
+				CurveRangeMap[SelectedKey.Curve].Include(SelectedTime);
+			}
+			else
+			{
+				CurveRangeMap.Add(SelectedKey.Curve, TInterval<float>(SelectedTime, SelectedTime));
+			}
+		}
+
+		if (CurveRangeMap.Num())
+		{
+			for (auto CurveToBake : CurveRangeMap)
+			{
+				CurveToBake.Key->RemoveRedundantKeys(ReduceTolerance, CurveToBake.Value.Min, CurveToBake.Value.Max);
+				ChangedCurveEditInfos.Add(GetViewModelForCurve(CurveToBake.Key)->CurveInfo);
+			}
+		}
+		else
+		{
+			for (auto CurveViewModel : CurveViewModels)
+			{
+				if (!bAnyCurveViewModelsSelected || CurveViewModel->bIsSelected)
+				{
+					CurveViewModel->CurveInfo.CurveToEdit->RemoveRedundantKeys(ReduceTolerance);
+					ChangedCurveEditInfos.Add(CurveViewModel->CurveInfo);
+				}
+			}
+		}
+
+		if (ChangedCurveEditInfos.Num())
+		{
+			CurveOwner->OnCurveChanged(ChangedCurveEditInfos);
+		}
+	}
+}
+
+
 void SCurveEditor::OnSelectPreInfinityExtrap(ERichCurveExtrapolation Extrapolation)
 {
 	const FScopedTransaction Transaction(LOCTEXT("CurveEditor_SetPreInfinityExtrapolation", "Set Pre-Infinity Extrapolation"));
@@ -3384,4 +3539,33 @@ TSharedPtr<FCurveViewModel> SCurveEditor::GetViewModelForCurve(FRichCurve* InCur
 	return TSharedPtr<FCurveViewModel>();
 }
 
+void SCurveEditor::GenericTextEntryModeless(const FText& DialogText, const FText& DefaultText, FOnTextCommitted OnTextComitted)
+{
+	TSharedRef<STextEntryPopup> TextEntryPopup = 
+		SNew(STextEntryPopup)
+		.Label(DialogText)
+		.DefaultText(DefaultText)
+		.OnTextCommitted(OnTextComitted)
+		.ClearKeyboardFocusOnCommit(false)
+		.SelectAllTextWhenFocused(true)
+		.MaxWidth(1024.0f);
+
+	EntryPopupMenu = FSlateApplication::Get().PushMenu(
+		SharedThis(this),
+		FWidgetPath(),
+		TextEntryPopup,
+		FSlateApplication::Get().GetCursorPos(),
+		FPopupTransitionEffect(FPopupTransitionEffect::TypeInPopup)
+	);
+}
+
+
+void SCurveEditor::CloseEntryPopupMenu()
+{
+	if (EntryPopupMenu.IsValid())
+	{
+		EntryPopupMenu.Pin()->Dismiss();
+	}
+}
+
 #undef LOCTEXT_NAMESPACE
diff --git a/Engine/Source/Editor/UnrealEd/Private/Toolkits/AssetEditorCommonCommands.cpp b/Engine/Source/Editor/UnrealEd/Private/Toolkits/AssetEditorCommonCommands.cpp
--- a/Engine/Source/Editor/UnrealEd/Private/Toolkits/AssetEditorCommonCommands.cpp
+++ b/Engine/Source/Editor/UnrealEd/Private/Toolkits/AssetEditorCommonCommands.cpp
@@ -8,6 +8,7 @@
 void FAssetEditorCommonCommands::RegisterCommands()
 {
 	UI_COMMAND( SaveAsset, "Save", "Saves this asset", EUserInterfaceActionType::Button, FInputChord() );
+	UI_COMMAND( SaveAssetAs, "Save As...", "Saves this asset under a different name", EUserInterfaceActionType::Button, FInputChord() );
 	UI_COMMAND( ReimportAsset, "Reimport", "Reimports the asset being edited", EUserInterfaceActionType::Button, FInputChord());
 	UI_COMMAND( SwitchToStandaloneEditor, "Switch to Standalone Editor", "Closes the level-centric asset editor and reopens it in 'standalone' mode", EUserInterfaceActionType::Button, FInputChord() );
 	UI_COMMAND( SwitchToWorldCentricEditor, "Switch to World-Centric Editor", "Closes the standalone asset editor and reopens it in 'world-centric' mode, docked within the level editor that it was originally opened in.", EUserInterfaceActionType::Button, FInputChord() );
diff --git a/Engine/Source/Editor/UnrealEd/Private/Toolkits/AssetEditorCommonCommands.h b/Engine/Source/Editor/UnrealEd/Private/Toolkits/AssetEditorCommonCommands.h
--- a/Engine/Source/Editor/UnrealEd/Private/Toolkits/AssetEditorCommonCommands.h
+++ b/Engine/Source/Editor/UnrealEd/Private/Toolkits/AssetEditorCommonCommands.h
@@ -17,6 +17,7 @@ public:
 	virtual void RegisterCommands() override;
 
 	TSharedPtr< FUICommandInfo > SaveAsset;
+	TSharedPtr< FUICommandInfo > SaveAssetAs;
 	TSharedPtr< FUICommandInfo > ReimportAsset;
 	TSharedPtr< FUICommandInfo > SwitchToStandaloneEditor;
 	TSharedPtr< FUICommandInfo > SwitchToWorldCentricEditor;
diff --git a/Engine/Source/Editor/UnrealEd/Private/Toolkits/AssetEditorToolkit.cpp b/Engine/Source/Editor/UnrealEd/Private/Toolkits/AssetEditorToolkit.cpp
--- a/Engine/Source/Editor/UnrealEd/Private/Toolkits/AssetEditorToolkit.cpp
+++ b/Engine/Source/Editor/UnrealEd/Private/Toolkits/AssetEditorToolkit.cpp
@@ -188,12 +188,16 @@ void FAssetEditorToolkit::InitAssetEditor( const EToolkitMode::Type Mode, const
 		Toolbar = SNullWidget::NullWidget;
 	}
 
-
 	ToolkitCommands->MapAction(
 		FAssetEditorCommonCommands::Get().SaveAsset,
 		FExecuteAction::CreateSP( this, &FAssetEditorToolkit::SaveAsset_Execute ),
 		FCanExecuteAction::CreateSP( this, &FAssetEditorToolkit::CanSaveAsset ));
 
+	ToolkitCommands->MapAction(
+		FAssetEditorCommonCommands::Get().SaveAssetAs,
+		FExecuteAction::CreateSP( this, &FAssetEditorToolkit::SaveAssetAs_Execute ),
+		FCanExecuteAction::CreateSP( this, &FAssetEditorToolkit::CanSaveAssetAs ));
+
 	ToolkitCommands->MapAction(
 		FGlobalEditorCommonCommands::Get().FindInContentBrowser,
 		FExecuteAction::CreateSP( this, &FAssetEditorToolkit::FindInContentBrowser_Execute ) );
@@ -422,6 +426,20 @@ const TArray< UObject* >& FAssetEditorToolkit::GetEditingObjects() const
 }
 
 
+void FAssetEditorToolkit::GetSaveableObjects(TArray<UObject*>& OutObjects) const
+{
+	for (const auto Object : EditingObjects)
+	{
+		// don't allow user to perform certain actions on objects that
+		// aren't actually assets (e.g. Level Script blueprint objects)
+		if ((Object != nullptr) && Object->IsAsset())
+		{
+			OutObjects.Add(Object);
+		}
+	}
+}
+
+
 void FAssetEditorToolkit::AddEditingObject(UObject* Object)
 {
 	EditingObjects.Add(Object);
@@ -438,24 +456,97 @@ void FAssetEditorToolkit::RemoveEditingObject(UObject* Object)
 
 void FAssetEditorToolkit::SaveAsset_Execute()
 {
-	if( ensure( EditingObjects.Num() > 0 ) )
+	if (EditingObjects.Num() == 0)
 	{
-		TArray< UPackage* > PackagesToSave;
+		return;
+	}
 
-		for( auto ObjectIter = EditingObjects.CreateConstIterator(); ObjectIter; ++ObjectIter )
+	TArray<UObject*> ObjectsToSave;
+	GetSaveableObjects(ObjectsToSave);
+
+	if (ObjectsToSave.Num() == 0)
+	{
+		return;
+	}
+
+	TArray<UPackage*> PackagesToSave;
+
+	for (auto Object : ObjectsToSave)
+	{
+		check((Object != nullptr) && Object->IsAsset());
+		PackagesToSave.Add(Object->GetOutermost());
+	}
+
+	FEditorFileUtils::PromptForCheckoutAndSave(PackagesToSave, bCheckDirtyOnAssetSave, /*bPromptToSave=*/ false);
+}
+
+
+void FAssetEditorToolkit::SaveAssetAs_Execute()
+{
+	if (EditingObjects.Num() == 0)
+	{
+		return;
+	}
+
+	TSharedPtr<IToolkitHost> MyToolkitHost = ToolkitHost.Pin();
+
+	if (!MyToolkitHost.IsValid())
+	{
+		return;
+	}
+
+	// get collection of objects to save
+	TArray<UObject*> ObjectsToSave;
+	GetSaveableObjects(ObjectsToSave);
+
+	if (ObjectsToSave.Num() == 0)
+	{
+		return;
+	}
+
+	// save assets under new name
+	TArray<UObject*> SavedObjects;
+	FEditorFileUtils::SaveAssetsAs(ObjectsToSave, SavedObjects);
+
+
+	// close existing asset editors for resaved assets
+	FAssetEditorManager& AssetEditorManager = FAssetEditorManager::Get();
+
+	/* @todo editor: Persona does not behave well when closing specific objects
+	for (int32 Index = 0; Index < ObjectsToSave.Num(); ++Index)
+	{
+		if ((SavedObjects[Index] != ObjectsToSave[Index]) && (SavedObjects[Index] != nullptr))
 		{
-			// Don't allow user to perform certain actions on objects that aren't actually assets (e.g. Level Script blueprint objects)
-			const auto EditingObject = *ObjectIter;
-			if( EditingObject != NULL && EditingObject->IsAsset() )
-			{
-				PackagesToSave.Add( EditingObject->GetOutermost() );
-			}
+			AssetEditorManager.CloseAllEditorsForAsset(ObjectsToSave[Index]);
 		}
+	}
 
-		FEditorFileUtils::PromptForCheckoutAndSave( PackagesToSave, bCheckDirtyOnAssetSave, /*bPromptToSave=*/ false );
+	// reopen asset editor
+	AssetEditorManager.OpenEditorForAssets(TArrayBuilder<UObject*>().Add(SavedObjects[0]), ToolkitMode, MyToolkitHost.ToSharedRef());
+	*/
+	// hack
+	TArray<UObject*> ObjectsToReopen;
+	for (auto Object : EditingObjects)
+	{
+		if (Object->IsAsset() && !ObjectsToSave.Contains(Object))
+		{
+			ObjectsToReopen.Add(Object);
+		}
+	}
+	for (auto Object : SavedObjects)
+	{
+		ObjectsToReopen.AddUnique(Object);
 	}
+	for (auto Object : EditingObjects)
+	{
+		AssetEditorManager.CloseAllEditorsForAsset(Object);
+		FAssetEditorManager::Get().NotifyAssetClosed(Object, this);
+	}
+	AssetEditorManager.OpenEditorForAssets(ObjectsToReopen, ToolkitMode, MyToolkitHost.ToSharedRef());
+	// end hack
 }
 
+
 const FSlateBrush* FAssetEditorToolkit::GetDefaultTabIcon() const
 {
 	if (EditingObjects.Num() == 0)
@@ -738,8 +829,8 @@ void FAssetEditorToolkit::FillDefaultFileMenuCommands( FMenuBuilder& MenuBuilder
 {
 	if( IsActuallyAnAsset() )
 	{
-		MenuBuilder.AddMenuEntry( FAssetEditorCommonCommands::Get().SaveAsset, "SaveAsset", TAttribute<FText>(), TAttribute<FText>(), FSlateIcon(FEditorStyle::GetStyleSetName(), "AssetEditor.SaveAsset.Greyscale") );
-
+		MenuBuilder.AddMenuEntry(FAssetEditorCommonCommands::Get().SaveAsset, "SaveAsset", TAttribute<FText>(), TAttribute<FText>(), FSlateIcon(FEditorStyle::GetStyleSetName(), "AssetEditor.SaveAsset.Greyscale"));
+		MenuBuilder.AddMenuEntry(FAssetEditorCommonCommands::Get().SaveAssetAs, "SaveAssetAs", TAttribute<FText>(), TAttribute<FText>(), FSlateIcon(FEditorStyle::GetStyleSetName(), "AssetEditor.SaveAssetAs.Small"));
 		MenuBuilder.AddMenuSeparator();
 	}
 
diff --git a/Engine/Source/Editor/UnrealEd/Public/FileHelpers.h b/Engine/Source/Editor/UnrealEd/Public/FileHelpers.h
--- a/Engine/Source/Editor/UnrealEd/Public/FileHelpers.h
+++ b/Engine/Source/Editor/UnrealEd/Public/FileHelpers.h
@@ -123,6 +123,14 @@ public:
 	 */
 	static UNREALED_API bool SaveLevel(ULevel* Level, const FString& DefaultFilename = TEXT( "" ) );
 
+	/**
+	 * Does a SaveAs for the specified assets.
+	 *
+	 * @param Assets The collection of assets to save.
+	 * @param SavedAssets The collection of corresponding saved assets (contains original asset if not resaved).
+	 */
+	UNREALED_API static void SaveAssetsAs(const TArray<UObject*>& Assets, TArray<UObject*>& OutSavedAssets);
+
 	/**
 	 * Does a saveAs for the specified level.
 	 *
diff --git a/Engine/Source/Editor/UnrealEd/Public/LevelEditorViewport.h b/Engine/Source/Editor/UnrealEd/Public/LevelEditorViewport.h
--- a/Engine/Source/Editor/UnrealEd/Public/LevelEditorViewport.h
+++ b/Engine/Source/Editor/UnrealEd/Public/LevelEditorViewport.h
@@ -420,6 +420,43 @@ public:
 		return ActorLockedToCamera;
 	}
 	
+	/**
+	 * Find the camera component to use for a locked actor
+	 * note this should return the same component as SLevelViewport::GetCameraInformationFromActor uses 
+	 */
+	UCameraComponent* GetCameraComponentForLockedActor(AActor const* Actor) const
+	{
+		if (Actor)
+		{
+			// see if actor has a camera component
+			TArray<UCameraComponent*> CamComps;
+			Actor->GetComponents<UCameraComponent>(CamComps);
+			for (UCameraComponent* Comp : CamComps)
+			{
+				if (Comp->bIsActive)
+				{
+					return Comp;
+				}
+			}
+
+			// now see if any actors are attached to us, directly or indirectly, that have an active camera component we might want to use
+			// we will just return the first one.
+			// #note: assumption here that attachment cannot be circular
+			TArray<AActor*> AttachedActors;
+			Actor->GetAttachedActors(AttachedActors);
+			for (AActor* AttachedActor : AttachedActors)
+			{
+				UCameraComponent* const Comp = GetCameraComponentForLockedActor(AttachedActor);
+				if (Comp)
+				{
+					return Comp;
+				}
+			}
+		}
+
+		return nullptr;
+	}
+
 	/** 
 	 * Find the camera component that is driving this viewport, in the following order of preference:
 	 *		1. Matinee locked actor
@@ -436,23 +473,7 @@ public:
 			LockedActor = ActorLockedToCamera.Get();
 		}
 
-		UCameraComponent* CameraComponent = nullptr;
-		if (LockedActor)
-		{
-			TArray<UCameraComponent*> CamComps;
-			LockedActor->GetComponents<UCameraComponent>(CamComps);
-
-			for (UCameraComponent* Comp : CamComps)
-			{
-				if (Comp->bIsActive)
-				{
-					CameraComponent = Comp;
-					break;
-				}
-			}
-		}
-
-		return CameraComponent;
+		return GetCameraComponentForLockedActor(LockedActor);
 	}
 
 	/** 
diff --git a/Engine/Source/Editor/UnrealEd/Public/RichCurveEditorCommands.h b/Engine/Source/Editor/UnrealEd/Public/RichCurveEditorCommands.h
--- a/Engine/Source/Editor/UnrealEd/Public/RichCurveEditorCommands.h
+++ b/Engine/Source/Editor/UnrealEd/Public/RichCurveEditorCommands.h
@@ -34,6 +34,9 @@ public:
 	TSharedPtr<FUICommandInfo> FlattenTangents;
 	TSharedPtr<FUICommandInfo> StraightenTangents;
 
+	TSharedPtr<FUICommandInfo> BakeCurve;
+	TSharedPtr<FUICommandInfo> ReduceCurve;
+
 	TSharedPtr<FUICommandInfo> SetPreInfinityExtrapCycle;
 	TSharedPtr<FUICommandInfo> SetPreInfinityExtrapCycleWithOffset;
 	TSharedPtr<FUICommandInfo> SetPreInfinityExtrapOscillate;
diff --git a/Engine/Source/Editor/UnrealEd/Public/SCurveEditor.h b/Engine/Source/Editor/UnrealEd/Public/SCurveEditor.h
--- a/Engine/Source/Editor/UnrealEd/Public/SCurveEditor.h
+++ b/Engine/Source/Editor/UnrealEd/Public/SCurveEditor.h
@@ -498,6 +498,12 @@ private:
 	/** Flatten or straighten tangents */
 	void OnFlattenOrStraightenTangents(bool bFlattenTangents);
 
+	/** Called when user selects bake or reduce curve */
+	void OnBakeCurve();
+	void OnBakeCurveSampleRateCommitted(const FText& InText, ETextCommit::Type CommitInfo);
+	void OnReduceCurve();
+	void OnReduceCurveToleranceCommitted(const FText& InText, ETextCommit::Type CommitInfo);
+
 	/** Called when the user selects the extrapolation type */
 	void OnSelectPreInfinityExtrap(ERichCurveExtrapolation Extrapolation);
 	bool IsPreInfinityExtrapSelected(ERichCurveExtrapolation Extrapolation);
@@ -577,6 +583,12 @@ protected:
 	/** Ensure that selected keys and tangents are still valid */
 	UNREALED_API void ValidateSelection();
 
+	/** Modeless Version of the String Entry Box */
+	void GenericTextEntryModeless(const FText& DialogText, const FText& DefaultText, FOnTextCommitted OnTextComitted);
+	
+	/** Closes the popup created by GenericTextEntryModeless*/
+	void CloseEntryPopupMenu();
+
 private:
 
 	/** User-supplied object for this curve editor */
@@ -733,6 +745,12 @@ protected:
 
 	/** The color used to draw the grid lines. */
 	FLinearColor GridColor;
+
+	/** The tolerance to use when reducing curves */
+	float ReduceTolerance;
+
+	/** Generic Popup Entry */
+	TWeakPtr<IMenu> EntryPopupMenu;
 };
 
 #endif // __SCurveEditor_h__
diff --git a/Engine/Source/Editor/UnrealEd/Public/Toolkits/AssetEditorToolkit.h b/Engine/Source/Editor/UnrealEd/Public/Toolkits/AssetEditorToolkit.h
--- a/Engine/Source/Editor/UnrealEd/Public/Toolkits/AssetEditorToolkit.h
+++ b/Engine/Source/Editor/UnrealEd/Public/Toolkits/AssetEditorToolkit.h
@@ -189,6 +189,9 @@ protected:
 	/**	Returns an array of all the objects currently being edited. Asserts if editing no objects */
 	const TArray< UObject* >& GetEditingObjects() const;
 
+	/** Get the collection of edited objects that can be saved. */
+	virtual void GetSaveableObjects(TArray<UObject*>& OutObjects) const;
+
 	/** Adds an item to the Editing Objects list */
 	virtual void AddEditingObject(UObject* Object);
 
@@ -201,6 +204,12 @@ protected:
 	/** Called when "Save" is clicked for this asset */
 	virtual void SaveAsset_Execute();
 
+	/** Called to test if "Save As" should be enabled for this asset */
+	virtual bool CanSaveAssetAs() const {return true;}
+
+	/** Called when "Save As" is clicked for this asset */
+	virtual void SaveAssetAs_Execute();
+
 	/** Called when "Find in Content Browser" is clicked for this asset */
 	virtual void FindInContentBrowser_Execute();
 
diff --git a/Engine/Source/Runtime/CinematicCamera/CinematicCamera.Build.cs b/Engine/Source/Runtime/CinematicCamera/CinematicCamera.Build.cs
new file mode 100644
--- /dev/null
+++ b/Engine/Source/Runtime/CinematicCamera/CinematicCamera.Build.cs
@@ -0,0 +1,23 @@
+// Copyright 1998-2016 Epic Games, Inc. All Rights Reserved.
+
+using UnrealBuildTool;
+
+public class CinematicCamera : ModuleRules
+{
+	public CinematicCamera(TargetInfo Target)
+	{
+		PrivateDependencyModuleNames.AddRange(
+			new string[] {
+				"Core",
+				"CoreUObject",
+                "Engine",
+			}
+		);
+
+        PrivateIncludePaths.AddRange(
+            new string[] {
+                "Runtime/CinematicCamera/Private"
+            })
+		;
+	}
+}
diff --git a/Engine/Source/Runtime/CinematicCamera/Private/CameraRig_Crane.cpp b/Engine/Source/Runtime/CinematicCamera/Private/CameraRig_Crane.cpp
new file mode 100644
--- /dev/null
+++ b/Engine/Source/Runtime/CinematicCamera/Private/CameraRig_Crane.cpp
@@ -0,0 +1,152 @@
+// Copyright 1998-2016 Epic Games, Inc. All Rights Reserved.
+
+#include "CinematicCameraPrivate.h"
+#include "CameraRig_Crane.h"
+
+#define LOCTEXT_NAMESPACE "CameraRig_Crane"
+
+ACameraRig_Crane::ACameraRig_Crane(const FObjectInitializer& ObjectInitializer)
+	: Super(ObjectInitializer)
+{
+	PrimaryActorTick.bCanEverTick = true;
+
+	CraneYaw = 0.f;
+	CranePitch = 0.f;
+	CraneArmLength = 500.f;
+
+	TransformComponent = CreateDefaultSubobject<USceneComponent>(TEXT("TransformComponent"));
+
+	// Make the scene component the root component
+	RootComponent = TransformComponent;
+	
+	// Setup camera defaults
+	CraneYawControl = CreateDefaultSubobject<USceneComponent>(TEXT("CraneYawControl"));
+	CraneYawControl->AttachParent = TransformComponent;
+	CraneYawControl->RelativeLocation = FVector(0.f, 0.f, 100.f);			// pivot height off the ground
+	CraneYawControl->RelativeRotation = FRotator(0.f, CraneYaw, 0.f);
+
+	CranePitchControl = CreateDefaultSubobject<USceneComponent>(TEXT("CranePitchControl"));
+	CranePitchControl->AttachParent = CraneYawControl;
+	CranePitchControl->RelativeRotation = FRotator(CranePitch, 0.f, 0.f);
+
+	CraneCameraMount = CreateDefaultSubobject<USceneComponent>(TEXT("CraneCameraMount"));
+	CraneCameraMount->AttachParent = CranePitchControl;
+	CraneCameraMount->RelativeLocation = FVector(CraneArmLength, 0.f, -35.f);			// negative z == underslung mount
+
+	// preview meshes
+	if (!IsRunningCommandlet() && !IsRunningDedicatedServer())
+	{
+		static ConstructorHelpers::FObjectFinder<UStaticMesh> CraneBaseMesh(TEXT("/Engine/BasicShapes/Cone"));
+		PreviewMesh_CraneBase = CreateOptionalDefaultSubobject<UStaticMeshComponent>(TEXT("PreviewMesh_CraneBase"));
+		if (PreviewMesh_CraneBase)
+		{
+			PreviewMesh_CraneBase->SetStaticMesh(CraneBaseMesh.Object);
+			PreviewMesh_CraneBase->AlwaysLoadOnClient = false;
+			PreviewMesh_CraneBase->AlwaysLoadOnServer = false;
+			PreviewMesh_CraneBase->SetCollisionProfileName(UCollisionProfile::NoCollision_ProfileName);
+			PreviewMesh_CraneBase->bHiddenInGame = true;
+			PreviewMesh_CraneBase->CastShadow = false;
+			PreviewMesh_CraneBase->PostPhysicsComponentTick.bCanEverTick = false;
+
+			PreviewMesh_CraneBase->AttachParent = RootComponent;		// sibling of yawcontrol
+			PreviewMesh_CraneBase->RelativeLocation = FVector(0.f, 0.f, 50.f);
+		}
+
+		static ConstructorHelpers::FObjectFinder<UStaticMesh> CraneArmMesh(TEXT("/Engine/BasicShapes/Cube"));
+		PreviewMesh_CraneArm = CreateOptionalDefaultSubobject<UStaticMeshComponent>(TEXT("PreviewMesh_CraneArm"));
+		if (PreviewMesh_CraneArm)
+		{
+			PreviewMesh_CraneArm->SetStaticMesh(CraneArmMesh.Object);
+			PreviewMesh_CraneArm->AlwaysLoadOnClient = false;
+			PreviewMesh_CraneArm->AlwaysLoadOnServer = false;
+			PreviewMesh_CraneArm->SetCollisionProfileName(UCollisionProfile::NoCollision_ProfileName);
+			PreviewMesh_CraneArm->bHiddenInGame = true;
+			PreviewMesh_CraneArm->CastShadow = false;
+			PreviewMesh_CraneArm->PostPhysicsComponentTick.bCanEverTick = false;
+			
+			PreviewMesh_CraneArm->AttachParent = CranePitchControl;		// sibling of the mount
+			PreviewMesh_CraneArm->RelativeScale3D = FVector(0.2f, 0.2f, 0.2f);
+		}
+
+		UpdatePreviewMeshes();
+	}
+}
+
+static const float CraneArmMesh_CounterweightOverhang = 100.f;
+
+void ACameraRig_Crane::UpdatePreviewMeshes()
+{
+	if (PreviewMesh_CraneArm)
+	{
+		// note to explain the math here:
+		// crane arm mesh is current a 100x100 cube. we want it to overhang the base pivot by 100cm
+		// (chosen arbitrarily)
+
+		float const TotalCraneArmMeshLength = CraneArmMesh_CounterweightOverhang + CraneArmLength;
+		float const CraneArmXScale = TotalCraneArmMeshLength / 100.f;
+		float const CraneArmXOffset = TotalCraneArmMeshLength * 0.5f - CraneArmMesh_CounterweightOverhang;
+
+		FVector NewRelScale3D = PreviewMesh_CraneArm->RelativeScale3D;
+		NewRelScale3D.X = CraneArmXScale;
+
+		FVector NewRelLoc = PreviewMesh_CraneArm->RelativeLocation;
+		NewRelLoc.X = CraneArmXOffset;
+
+		PreviewMesh_CraneArm->SetRelativeScale3D(NewRelScale3D);
+		PreviewMesh_CraneArm->SetRelativeLocation(NewRelLoc);
+	}
+}
+
+void ACameraRig_Crane::UpdateCraneComponents()
+{
+	FRotator NewYawControlRot = CraneYawControl->RelativeRotation;
+	NewYawControlRot.Yaw = CraneYaw;
+	CraneYawControl->SetRelativeRotation(FRotator(0.f, CraneYaw, 0.f));
+
+	FRotator NewPitchControlRot = CranePitchControl->RelativeRotation;
+	NewPitchControlRot.Pitch = CranePitch;
+	CranePitchControl->SetRelativeRotation(FRotator(CranePitch, 0.f, 0.f));
+
+	FVector NewCameraMountLoc = CraneCameraMount->RelativeLocation;
+	NewCameraMountLoc.X = CraneArmLength;
+	CraneCameraMount->SetRelativeLocation(NewCameraMountLoc);
+
+	// zero the pitch from the camera mount component
+	// this effectively gives us bAbsoluteRotation for only pitch component of an attached camera
+	FRotator NewCameraMountWorldRot = CraneCameraMount->GetComponentRotation();
+	NewCameraMountWorldRot.Pitch = 0.f;
+	CraneCameraMount->SetWorldRotation(NewCameraMountWorldRot);
+
+	UpdatePreviewMeshes();
+}
+
+
+void ACameraRig_Crane::Tick(float DeltaTime)
+{
+	Super::Tick(DeltaTime);
+
+	// feed exposed API into underlying components
+	UpdateCraneComponents();
+}
+
+#if WITH_EDITOR
+void ACameraRig_Crane::PostEditChangeProperty(FPropertyChangedEvent& PropertyChangedEvent)
+{
+	Super::PostEditChangeProperty(PropertyChangedEvent);
+
+	UpdateCraneComponents();
+}
+
+USceneComponent* ACameraRig_Crane::GetDefaultAttachComponent() const
+{
+	return CraneCameraMount;
+}
+#endif // WITH_EDITOR
+
+bool ACameraRig_Crane::ShouldTickIfViewportsOnly() const
+{
+	return true;
+}
+
+
+#undef LOCTEXT_NAMESPACE
\ No newline at end of file
diff --git a/Engine/Source/Runtime/CinematicCamera/Private/CineCameraActor.cpp b/Engine/Source/Runtime/CinematicCamera/Private/CineCameraActor.cpp
new file mode 100644
--- /dev/null
+++ b/Engine/Source/Runtime/CinematicCamera/Private/CineCameraActor.cpp
@@ -0,0 +1,93 @@
+// Copyright 1998-2015 Epic Games, Inc. All Rights Reserved.
+
+#include "CinematicCameraPrivate.h"
+#include "CineCameraComponent.h"
+#include "CineCameraActor.h"
+
+#define LOCTEXT_NAMESPACE "CineCameraActor"
+
+//////////////////////////////////////////////////////////////////////////
+// ACameraActor
+
+ACineCameraActor::ACineCameraActor(const FObjectInitializer& ObjectInitializer)
+	: Super(ObjectInitializer
+			.SetDefaultSubobjectClass<UCineCameraComponent>(TEXT("CameraComponent"))
+	)
+{
+	PrimaryActorTick.bCanEverTick = true;
+	SetActorTickEnabled(LookatTrackingSettings.bEnableLookAtTracking);
+}
+
+#if WITH_EDITOR
+void ACineCameraActor::PostEditChangeProperty(FPropertyChangedEvent& PropertyChangedEvent)
+{
+	Super::PostEditChangeProperty(PropertyChangedEvent);
+
+	SetActorTickEnabled(LookatTrackingSettings.bEnableLookAtTracking);
+}
+#endif
+
+void ACineCameraActor::PostInitializeComponents()
+{
+	Super::PostInitializeComponents();
+	SetActorTickEnabled(LookatTrackingSettings.bEnableLookAtTracking);
+
+	LookatTrackingSettings.LastLookatTrackingRotation = GetActorRotation();
+}
+
+bool ACineCameraActor::ShouldTickIfViewportsOnly() const
+{
+	return true;
+}
+
+FVector ACineCameraActor::GetLookatLocation() const
+{
+	FVector FinalLookat;
+	if (LookatTrackingSettings.ActorToTrack)
+	{
+		FTransform const BaseTransform = LookatTrackingSettings.ActorToTrack ? LookatTrackingSettings.ActorToTrack->GetActorTransform() : FTransform::Identity;
+		FinalLookat = BaseTransform.TransformPosition(LookatTrackingSettings.RelativeOffset);
+	}
+	else
+	{
+		FinalLookat = LookatTrackingSettings.RelativeOffset;
+	}
+
+	return FinalLookat;
+}
+
+static const FColor DebugLookatTrackingPointSolidColor(200, 200, 32, 128);		// yellow
+static const FColor DebugLookatTrackingPointOutlineColor = FColor::Black;
+
+void ACineCameraActor::Tick(float DeltaTime)
+{
+	Super::Tick(DeltaTime);
+
+	if (CameraComponent && LookatTrackingSettings.bEnableLookAtTracking)
+	{
+		// do the lookat tracking
+		// #note this will turn the whole actor, which assumes the cameracomponent's transform is the same as the root component
+		// more complex component hierarchies will require different handling here
+		FVector const LookatLoc = GetLookatLocation();
+		FVector const ToLookat = LookatLoc - GetActorLocation();
+		FRotator const FinalRot = FMath::RInterpTo(LookatTrackingSettings.LastLookatTrackingRotation, ToLookat.Rotation(), DeltaTime, LookatTrackingSettings.LookAtTrackingInterpSpeed);
+		SetActorRotation(FinalRot);
+
+		// we store this ourselves in case other systems try to change our rotation, and end fighting the interpolation
+		LookatTrackingSettings.LastLookatTrackingRotation = FinalRot;
+
+		if (LookatTrackingSettings.bDrawDebugLookAtTrackingPosition)
+		{
+			::DrawDebugSolidBox(GetWorld(), LookatLoc, FVector(12.f), DebugLookatTrackingPointSolidColor);
+			::DrawDebugBox(GetWorld(), LookatLoc, FVector(12.f), DebugLookatTrackingPointOutlineColor);
+		}
+	}
+	else
+	{
+		// no tracking, no ticking
+		SetActorTickEnabled(false);
+	}
+}
+
+
+#undef LOCTEXT_NAMESPACE
diff --git a/Engine/Source/Runtime/CinematicCamera/Private/CineCameraComponent.cpp b/Engine/Source/Runtime/CinematicCamera/Private/CineCameraComponent.cpp
new file mode 100644
--- /dev/null
+++ b/Engine/Source/Runtime/CinematicCamera/Private/CineCameraComponent.cpp
@@ -0,0 +1,249 @@
+// Copyright 1998-2016 Epic Games, Inc. All Rights Reserved.
+
+#include "CinematicCameraPrivate.h"
+#include "CineCameraComponent.h"
+
+#define LOCTEXT_NAMESPACE "CineCameraComponent"
+
+
+//////////////////////////////////////////////////////////////////////////
+// UCameraComponent
+
+UCineCameraComponent::UCineCameraComponent()
+{
+	// Super 35mm 4 Perf
+	FilmbackSettings.SensorWidth = 24.89f;
+	FilmbackSettings.SensorHeight = 18.67;
+	LensSettings.MinFocalLength = 50.f;
+	LensSettings.MaxFocalLength = 50.f;
+	LensSettings.MinFStop = 2.f;
+	LensSettings.MaxFStop = 2.f;
+	LensSettings.MinimumFocusDistance = 15.f;
+	LensSettings.MaxReproductionRatio = 1 / 11.f;
+
+	bConstrainAspectRatio = true;
+
+	RecalcDerivedData();
+}
+
+void UCineCameraComponent::PostLoad()
+{
+	RecalcDerivedData();
+	Super::PostLoad();
+}
+
+#if WITH_EDITORONLY_DATA
+void UCineCameraComponent::PostEditChangeProperty(FPropertyChangedEvent& PropertyChangedEvent)
+{
+	RecalcDerivedData();
+
+	Super::PostEditChangeProperty(PropertyChangedEvent);
+}
+#endif	// WITH_EDITORONLY_DATA
+
+
+float UCineCameraComponent::GetHorizontalFieldOfView() const
+{
+	return (CurrentFocalLength > 0.f)
+		? FMath::RadiansToDegrees(2.f * FMath::Atan(FilmbackSettings.SensorWidth / (2.f * CurrentFocalLength)))
+		: 0.f;
+}
+
+float UCineCameraComponent::GetVerticalFieldOfView() const
+{
+	return (CurrentFocalLength > 0.f)
+		? FMath::RadiansToDegrees(2.f * FMath::Atan(FilmbackSettings.SensorHeight / (2.f * CurrentFocalLength)))
+		: 0.f;
+}
+
+float UCineCameraComponent::GetWorldToMetersScale() const
+{
+	UWorld const* const World = GetWorld();
+	AWorldSettings const* const WorldSettings = World ? World->GetWorldSettings() : nullptr;
+	return WorldSettings ? WorldSettings->WorldToMeters : 100.f;
+}
+
+// static
+TArray<FNamedFilmbackPreset> const& UCineCameraComponent::GetFilmbackPresets()
+{
+	return GetDefault<UCineCameraComponent>()->FilmbackPresets;
+}
+
+// static
+TArray<FNamedLensPreset> const& UCineCameraComponent::GetLensPresets()
+{
+	return GetDefault<UCineCameraComponent>()->LensPresets;
+}
+
+void UCineCameraComponent::RecalcDerivedData()
+{
+	// respect physical limits of the (simulated) hardware
+	CurrentFocalLength = FMath::Clamp(CurrentFocalLength, LensSettings.MinFocalLength, LensSettings.MaxFocalLength);
+	CurrentAperture = FMath::Clamp(CurrentAperture, LensSettings.MinFStop, LensSettings.MaxFStop);
+
+	float const MinFocusDistInWorldUnits = LensSettings.MinimumFocusDistance * (GetWorldToMetersScale() / 1000.f);	// convert mm to uu
+	FocusSettings.ManualFocusDistance = FMath::Max(FocusSettings.ManualFocusDistance, MinFocusDistInWorldUnits);
+
+	FieldOfView = GetHorizontalFieldOfView();
+	FilmbackSettings.SensorAspectRatio = (FilmbackSettings.SensorHeight > 0.f) ? (FilmbackSettings.SensorWidth / FilmbackSettings.SensorHeight) : 0.f;
+	AspectRatio = FilmbackSettings.SensorAspectRatio;
+}
+
+
+static const FName NAME_GetDesiredFocusDistance(TEXT("GetDesiredFocusDistance"));
+static const float SpotFocusTraceDist = 1000000.f;
+static const float FocusDistInfinity = 1000000.f;
+
+float UCineCameraComponent::GetDesiredFocusDistance(FMinimalViewInfo& DesiredView) const
+{
+	float DesiredFocusDistance = 0.f;
+
+	// get focus distance
+	switch (FocusSettings.FocusMethod)
+	{
+	case ECameraFocusMethod::Manual:
+		DesiredFocusDistance = FocusSettings.ManualFocusDistance;
+		break;
+
+	case ECameraFocusMethod::Spot:
+		{
+			// AutoFocus, do a trace
+			AActor* const OwningActor = GetOwner();
+
+			// trace to get depth at center of screen
+			FCollisionQueryParams TraceParams(NAME_GetDesiredFocusDistance, true, OwningActor);
+			FHitResult Hit;
+
+			FVector const TraceStart = DesiredView.Location;
+			FVector const TraceEnd = TraceStart + DesiredView.Rotation.Vector() * SpotFocusTraceDist;
+			bool const bHit = GetWorld()->LineTraceSingleByChannel(Hit, TraceStart, TraceEnd, ECC_WorldStatic, TraceParams);
+			DesiredFocusDistance = bHit ? (TraceStart - Hit.ImpactPoint).Size() : FocusDistInfinity;
+		}
+		break;
+
+	case ECameraFocusMethod::Tracking:
+		{
+			AActor const* const TrackedActor = FocusSettings.TrackingFocusSettings.ActorToTrack;
+
+			FVector FocusPoint;
+			if (TrackedActor)
+			{
+				FTransform const BaseTransform = TrackedActor->GetActorTransform();
+				FocusPoint = BaseTransform.TransformPosition(FocusSettings.TrackingFocusSettings.RelativeOffset);
+			}
+			else
+			{
+				FocusPoint = FocusSettings.TrackingFocusSettings.RelativeOffset;
+			}
+
+			DesiredFocusDistance = (FocusPoint - DesiredView.Location).Size();
+
+// 			// try to get the position of the tracked object
+// 			UObject* const TrackedObj = FocusSettings.TrackingFocusSettings.ObjectToTrack;
+// 
+// 			// infinity is the fallback in case none of the other code paths produces a good result
+// 			DesiredFocusDistance = FocusDistInfinity;
+// 			if (TrackedObj)
+// 			{
+// 				// we have a socket/bone, see if the object is a scenecomponent
+// 				USceneComponent* const TrackedSceneComp = Cast<USceneComponent>(TrackedObj);
+// 				if (TrackedSceneComp)
+// 				{
+// 					if (FocusSettings.TrackingFocusSettings.BoneOrSocketName != NAME_None)
+// 					{
+// 						FVector const SocketLoc = TrackedSceneComp->GetSocketLocation(FocusSettings.TrackingFocusSettings.BoneOrSocketName);
+// 						DesiredFocusDistance = (SocketLoc - DesiredView.Location).Size();
+// 					}
+// 					else
+// 					{
+// 						DesiredFocusDistance = (TrackedSceneComp->GetComponentLocation() - DesiredView.Location).Size();
+// 					}
+// 				}
+// 				else
+// 				{
+// 					AActor* const TrackedActor = Cast<AActor>(TrackedObj);
+// 					if (TrackedActor)
+// 					{
+// 						if (FocusSettings.TrackingFocusSettings.BoneOrSocketName != NAME_None)
+// 						{
+// 							// user asked for an actor and a bone or socket, try to find a suitable component
+// 							TArray<USceneComponent*> OutComponents;
+// 							TrackedActor->GetComponents(OutComponents);
+// 							for (USceneComponent const* Comp : OutComponents)
+// 							{
+// 								if (Comp->DoesSocketExist(FocusSettings.TrackingFocusSettings.BoneOrSocketName))
+// 								{
+// 									FVector const SocketLoc = Comp->GetSocketLocation(FocusSettings.TrackingFocusSettings.BoneOrSocketName);
+// 									DesiredFocusDistance = (SocketLoc - DesiredView.Location).Size();
+// 									break;
+// 								}
+// 							}
+// 
+// 						}
+// 						else
+// 						{
+// 							DesiredFocusDistance = (TrackedActor->GetActorLocation() - DesiredView.Location).Size();
+// 						}
+// 					}
+// 				}
+// 			}
+		}
+		break;
+	}
+	
+	// add in the adjustment offset
+	DesiredFocusDistance += FocusSettings.FocusOffset;
+
+	return DesiredFocusDistance;
+}
+
+static const FColor DebugFocusPlaneColor(64, 32, 128, 64);		// purple
+
+void UCineCameraComponent::GetCameraView(float DeltaTime, FMinimalViewInfo& DesiredView)
+{
+	RecalcDerivedData();
+
+	Super::GetCameraView(DeltaTime, DesiredView);
+
+	UpdateCameraLens(DeltaTime, DesiredView);
+
+	if (FocusSettings.bDrawDebugFocusPlane)
+	{
+		FVector const CamDir = DesiredView.Rotation.Vector();
+		FVector const FocusPoint = DesiredView.Location + CamDir * CurrentFocusDistance;
+
+		FPlane P(FocusPoint, CamDir);
+
+		::DrawDebugSolidPlane(GetWorld(), P, DesiredView.Location, 10000.f, DebugFocusPlaneColor, false, -1, SDPG_World);
+	}
+}
+
+
+void UCineCameraComponent::UpdateCameraLens(float DeltaTime, FMinimalViewInfo& DesiredView)
+{
+	// Update focus/DoF
+	DesiredView.PostProcessBlendWeight = 1.f;
+	DesiredView.PostProcessSettings.bOverride_DepthOfFieldMethod = true;
+	DesiredView.PostProcessSettings.DepthOfFieldMethod = DOFM_CircleDOF;
+
+	DesiredView.PostProcessSettings.bOverride_DepthOfFieldFstop = true;
+	DesiredView.PostProcessSettings.DepthOfFieldFstop = CurrentAperture;
+
+	CurrentFocusDistance = GetDesiredFocusDistance(DesiredView);
+
+	// clamp to min focus distance
+	float const MinFocusDistInWorldUnits = LensSettings.MinimumFocusDistance * (GetWorldToMetersScale() / 1000.f);	// convert mm to uu
+	CurrentFocusDistance = FMath::Max(CurrentFocusDistance, MinFocusDistInWorldUnits);
+
+	// smoothing, if desired
+ 	if (FocusSettings.bSmoothFocusChanges)
+ 	{
+		CurrentFocusDistance = FMath::FInterpTo(LastFocusDistance, CurrentFocusDistance, DeltaTime, FocusSettings.FocusSmoothingInterpSpeed);
+ 	}
+	LastFocusDistance = CurrentFocusDistance;
+
+	DesiredView.PostProcessSettings.bOverride_DepthOfFieldFocalDistance = true;
+	DesiredView.PostProcessSettings.DepthOfFieldFocalDistance = CurrentFocusDistance;
+}
+
+#undef LOCTEXT_NAMESPACE
diff --git a/Engine/Source/Runtime/CinematicCamera/Private/CinematicCamera.h b/Engine/Source/Runtime/CinematicCamera/Private/CinematicCamera.h
new file mode 100644
--- /dev/null
+++ b/Engine/Source/Runtime/CinematicCamera/Private/CinematicCamera.h
@@ -0,0 +1,8 @@
+// Copyright 1998-2016 Epic Games, Inc. All Rights Reserved.
+
+
+#pragma once
+#include "Engine.h"
+#include "StaticMeshResources.h"
+#include "FoliageType_InstancedStaticMesh.h"
+#include "InstancedFoliageActor.h"
diff --git a/Engine/Source/Runtime/CinematicCamera/Private/CinematicCameraModule.cpp b/Engine/Source/Runtime/CinematicCamera/Private/CinematicCameraModule.cpp
new file mode 100644
--- /dev/null
+++ b/Engine/Source/Runtime/CinematicCamera/Private/CinematicCameraModule.cpp
@@ -0,0 +1,6 @@
+// Copyright 1998-2016 Epic Games, Inc. All Rights Reserved.
+
+#include "CinematicCameraPrivate.h"
+#include "CinematicCameraModule.h"
+
+IMPLEMENT_MODULE( ICinematicCameraModule, CinematicCamera );
diff --git a/Engine/Source/Runtime/CinematicCamera/Private/CinematicCameraPrivate.h b/Engine/Source/Runtime/CinematicCamera/Private/CinematicCameraPrivate.h
new file mode 100644
--- /dev/null
+++ b/Engine/Source/Runtime/CinematicCamera/Private/CinematicCameraPrivate.h
@@ -0,0 +1,5 @@
+// Copyright 1998-2016 Epic Games, Inc. All Rights Reserved.
+
+#pragma once
+
+#include "Engine.h"
diff --git a/Engine/Source/Runtime/CinematicCamera/Public/CameraRig_Crane.h b/Engine/Source/Runtime/CinematicCamera/Public/CameraRig_Crane.h
new file mode 100644
--- /dev/null
+++ b/Engine/Source/Runtime/CinematicCamera/Public/CameraRig_Crane.h
@@ -0,0 +1,66 @@
+// Copyright 1998-2016 Epic Games, Inc. All Rights Reserved.
+
+#pragma once
+
+//#include "Engine/Scene.h"
+
+#include "CameraRig_Crane.generated.h"
+
+class USceneComponent;
+
+/** 
+ * 
+ */
+UCLASS(Blueprintable)
+class CINEMATICCAMERA_API ACameraRig_Crane : public AActor
+{
+	GENERATED_BODY()
+	
+public:
+
+	// ctor
+	ACameraRig_Crane(const FObjectInitializer& ObjectInitialier);
+
+	virtual void Tick(float DeltaTime) override;
+	virtual bool ShouldTickIfViewportsOnly() const override;
+
+	UPROPERTY(Interp, EditAnywhere, BlueprintReadWrite, Category = "Crane Controls", meta = (UIMin = "-360", UIMax = "360", Units = deg))
+	float CranePitch;
+	
+	UPROPERTY(Interp, EditAnywhere, BlueprintReadWrite, Category = "Crane Controls", meta = (UIMin = "-360", UIMax = "360", Units = deg))
+	float CraneYaw;
+	
+	UPROPERTY(Interp, EditAnywhere, BlueprintReadWrite, Category = "Crane Controls", meta = (ClampMin=0, Units = cm))
+	float CraneArmLength;
+
+#if WITH_EDITOR
+	virtual void PostEditChangeProperty(FPropertyChangedEvent& PropertyChangedEvent) override;
+	virtual class USceneComponent* GetDefaultAttachComponent() const override;
+#endif
+	
+protected:
+	void UpdatePreviewMeshes();
+	void UpdateCraneComponents();
+
+private:
+
+	/** Root component to give the whole actor a transform. */
+	UPROPERTY(EditDefaultsOnly, Category = "Crane Components")
+	USceneComponent* TransformComponent;
+
+	UPROPERTY(EditDefaultsOnly, Category = "Crane Components")
+	USceneComponent* CraneYawControl;
+
+	UPROPERTY(EditDefaultsOnly, Category = "Crane Components")
+	USceneComponent* CranePitchControl;
+
+	UPROPERTY(EditDefaultsOnly, Category = "Crane Components")
+	USceneComponent* CraneCameraMount;
+
+	UPROPERTY()
+	UStaticMeshComponent* PreviewMesh_CraneArm;
+
+	UPROPERTY()
+	UStaticMeshComponent* PreviewMesh_CraneBase;
+
+};
diff --git a/Engine/Source/Runtime/CinematicCamera/Public/CameraRig_Rail.h b/Engine/Source/Runtime/CinematicCamera/Public/CameraRig_Rail.h
new file mode 100644
--- /dev/null
+++ b/Engine/Source/Runtime/CinematicCamera/Public/CameraRig_Rail.h
@@ -0,0 +1,37 @@
+// Copyright 1998-2016 Epic Games, Inc. All Rights Reserved.
+
+#pragma once
+
+#include "CameraRig_Rail.generated.h"
+
+class USceneComponent;
+
+/** 
+ * 
+ */
+UCLASS(Blueprintable)
+class CINEMATICCAMERA_API ACameraRig_Rail : public AActor
+{
+	GENERATED_BODY()
+	
+public:
+	// ctor
+	ACameraRig_Rail(const FObjectInitializer& ObjectInitialier);
+
+	UPROPERTY(Interp, EditAnywhere, BlueprintReadWrite, Category = "Crane Controls")
+	float DistanceAlongRail;
+	
+#if WITH_EDITOR
+	virtual void PostEditChangeProperty(FPropertyChangedEvent& PropertyChangedEvent) override;
+	virtual class USceneComponent* GetDefaultAttachComponent() const override;
+#endif
+	
+private:
+
+// 	/** Root component to give the whole actor a transform. */
+// 	UPROPERTY()
+// 	USceneComponent* TransformComponent;
+
+	UPROPERTY()
+	USceneComponent* CraneRailMount;
+};
diff --git a/Engine/Source/Runtime/CinematicCamera/Public/CineCameraActor.h b/Engine/Source/Runtime/CinematicCamera/Public/CineCameraActor.h
new file mode 100644
--- /dev/null
+++ b/Engine/Source/Runtime/CinematicCamera/Public/CineCameraActor.h
@@ -0,0 +1,60 @@
+// Copyright 1998-2015 Epic Games, Inc. All Rights Reserved.
+
+#pragma once
+
+#include "CineCameraActor.generated.h"
+
+USTRUCT()
+struct FCameraLookatTrackingSettings
+{
+	GENERATED_USTRUCT_BODY()
+
+	/** True to enable lookat tracking, false otherwise. */
+	UPROPERTY(/*Interp, */EditAnywhere, BlueprintReadWrite, Category = "LookAt")
+	uint8 bEnableLookAtTracking : 1;
+
+	UPROPERTY(Transient, EditAnywhere, BlueprintReadWrite, Category = "LookAt")
+	uint8 bDrawDebugLookAtTrackingPosition : 1;
+
+	/** Controls degree of smoothing. 0.f for no smoothing, higher numbers for faster/tighter tracking. */
+	UPROPERTY(EditAnywhere, BlueprintReadWrite, Category = "LookAt")
+	float LookAtTrackingInterpSpeed;
+
+	/** Last known lookat tracking rotation (used during interpolation) */
+	FRotator LastLookatTrackingRotation;
+
+	UPROPERTY(/*Interp, */EditAnywhere, BlueprintReadWrite, Category = "LookAt")
+	AActor* ActorToTrack;
+
+	/** Offset from actor position to look at. Relative to actor if tracking an actor, relative to world otherwise. */
+	UPROPERTY(Interp, EditAnywhere, BlueprintReadWrite, Category = "LookAt")
+	FVector RelativeOffset;
+};
+
+/** 
+ * A CineCameraActor is a CameraActor specialized to work like a cinematic camera.
+ */
+UCLASS(ClassGroup = Common, hideCategories = (Input, Rendering, AutoPlayerActivation), showcategories = ("Input|MouseInput", "Input|TouchInput"), Blueprintable)
+class CINEMATICCAMERA_API ACineCameraActor : public ACameraActor
+{
+	GENERATED_BODY()
+
+public:
+	// Ctor
+	ACineCameraActor(const FObjectInitializer& ObjectInitializer);
+
+	virtual void Tick(float DeltaTime) override;
+	virtual bool ShouldTickIfViewportsOnly() const override;
+	virtual void PostInitializeComponents() override;
+
+	UPROPERTY(EditAnywhere, BlueprintReadWrite, Category = "LookAt")
+	FCameraLookatTrackingSettings LookatTrackingSettings;
+
+#if WITH_EDITOR
+	virtual void PostEditChangeProperty(FPropertyChangedEvent& PropertyChangedEvent) override;
+#endif
+
+protected:
+	FVector GetLookatLocation() const;
+
+};
diff --git a/Engine/Source/Runtime/CinematicCamera/Public/CineCameraComponent.h b/Engine/Source/Runtime/CinematicCamera/Public/CineCameraComponent.h
new file mode 100644
--- /dev/null
+++ b/Engine/Source/Runtime/CinematicCamera/Public/CineCameraComponent.h
@@ -0,0 +1,245 @@
+// Copyright 1998-2016 Epic Games, Inc. All Rights Reserved.
+
+#pragma once
+#include "Camera/CameraComponent.h"
+#include "CineCameraComponent.generated.h"
+
+USTRUCT()
+struct FCameraFilmbackSettings
+{
+	GENERATED_USTRUCT_BODY()
+
+	UPROPERTY(EditAnywhere, BlueprintReadWrite, Category = "Filmback", meta = (ClampMin = "0.001", ForceUnits = mm))
+	float SensorWidth;
+
+	/** Vertical size of filmback or digital sensor, in mm. */
+	UPROPERTY(EditAnywhere, BlueprintReadWrite, Category = "Filmback", meta = (ClampMin = "0.001", ForceUnits = mm))
+	float SensorHeight;
+
+	UPROPERTY(VisibleAnywhere, BlueprintReadOnly, Category = "Filmback")
+	float SensorAspectRatio;
+};
+
+USTRUCT()
+struct FNamedFilmbackPreset
+{
+	GENERATED_USTRUCT_BODY()
+
+	UPROPERTY()
+	FString Name;
+
+	UPROPERTY()
+	FCameraFilmbackSettings FilmbackSettings;
+};
+
+USTRUCT()
+struct FCameraLensSettings
+{
+	GENERATED_USTRUCT_BODY()
+
+	UPROPERTY(EditAnywhere, BlueprintReadWrite, Category = "Lens", meta = (ForceUnits = mm))
+	float MinFocalLength;
+
+	UPROPERTY(EditAnywhere, BlueprintReadWrite, Category = "Lens", meta = (ForceUnits = mm))
+	float MaxFocalLength;
+
+	// #todo, split into MinFStopAtMinFocalLength, MinFStopAtMaxFocalLength for variable-min-fstop lenses
+
+	UPROPERTY(EditAnywhere, BlueprintReadWrite, Category = "Lens")
+	float MinFStop;
+
+	UPROPERTY(EditAnywhere, BlueprintReadWrite, Category = "Lens")
+	float MaxFStop;
+
+	UPROPERTY(EditAnywhere, BlueprintReadWrite, Category = "Lens", meta = (ForceUnits = mm))
+	float MinimumFocusDistance;
+
+	UPROPERTY(EditAnywhere, BlueprintReadWrite, Category = "Lens")
+	float MaxReproductionRatio;
+
+	bool operator==(const FCameraLensSettings& Other) const
+	{
+		return (MinFocalLength == Other.MinFocalLength)
+			&& (MaxFocalLength == Other.MaxFocalLength)
+			&& (MinFStop == Other.MinFStop)
+			&& (MaxFStop == Other.MaxFStop)
+			&& (MinimumFocusDistance == Other.MinimumFocusDistance)
+			&& (MaxReproductionRatio == Other.MaxReproductionRatio);
+	}
+};
+
+
+USTRUCT()
+struct FNamedLensPreset
+{
+	GENERATED_USTRUCT_BODY()
+
+	UPROPERTY()
+	FString Name;
+
+	UPROPERTY()
+	FCameraLensSettings LensSettings;
+};
+
+UENUM()
+enum class ECameraFocusMethod : uint8
+{
+	None,		/** Disable DoF entirely. */
+	Manual,		/** Allows for specifying or animating exact focus distances. */
+	Spot UMETA(Hidden), /** Raycasts into the scene, focuses on the hit location. */		// #todo, make this work
+	Tracking,	/** Locks focus to specific object. */
+};
+
+USTRUCT()
+struct FCameraSpotFocusSettings
+{
+	GENERATED_USTRUCT_BODY()
+
+	UPROPERTY(EditAnywhere, BlueprintReadWrite, Category="Spot Focus")
+	TArray<FVector2D> NormalizedSpotCoordinates;
+	
+	FCameraSpotFocusSettings()
+	{
+		NormalizedSpotCoordinates.Add(FVector2D(0.5f, 0.5f));
+	}
+};
+
+USTRUCT()
+struct FCameraTrackingFocusSettings
+{
+	GENERATED_USTRUCT_BODY()
+
+	UPROPERTY(EditAnywhere, BlueprintReadWrite, Category = "Tracking Focus")
+	AActor* ActorToTrack;
+
+	// #todo, ability to pick a component and socket/bone?
+
+// 	UPROPERTY(EditAnywhere, BlueprintReadWrite, Category = "Tracking Focus")
+// 	USceneComponent* ComponentToTrack;
+// 	
+// 	UPROPERTY(EditAnywhere, BlueprintReadWrite, Category = "Tracking Focus")
+// 	FName BoneOrSocketName;
+
+	/** Offset from actor position to look at. Relative to actor if tracking an actor, relative to world otherwise. */
+	UPROPERTY(Interp, EditAnywhere, BlueprintReadWrite, Category = "LookAt")
+	FVector RelativeOffset;
+
+	FCameraTrackingFocusSettings()
+		: ActorToTrack(nullptr)
+	{}
+};
+
+
+USTRUCT()
+struct FCameraFocusSettings
+{
+	GENERATED_USTRUCT_BODY()
+
+	UPROPERTY(EditAnywhere, BlueprintReadWrite, Category = "Focus Settings")
+	TEnumAsByte<ECameraFocusMethod> FocusMethod;
+	
+	/** True to draw a translucent plane at the current focus depth, for easy tweaking when not looking through the camera. */
+	UPROPERTY(Transient, EditAnywhere, BlueprintReadWrite, Category = "Focus Settings")
+	uint8 bDrawDebugFocusPlane : 1;
+
+	UPROPERTY(EditAnywhere, BlueprintReadWrite, Category = "Focus Settings")
+	uint8 bSmoothFocusChanges : 1;
+	
+	UPROPERTY(/*EditAnywhere, BlueprintReadWrite, Category = "Focus Settings"*/)
+	uint8 bSimulateFocusBreathing : 1;
+
+	UPROPERTY(Interp, EditAnywhere, BlueprintReadWrite, Category = "Focus Settings")
+	float ManualFocusDistance;
+
+	UPROPERTY(/*EditAnywhere, BlueprintReadWrite, Category = "Focus Settings"*/)
+	FCameraSpotFocusSettings SpotFocusSettings;
+
+	UPROPERTY(EditAnywhere, BlueprintReadWrite, Category = "Focus Settings")
+	FCameraTrackingFocusSettings TrackingFocusSettings;
+
+	UPROPERTY(EditAnywhere, BlueprintReadWrite, Category = "Focus Settings")
+	float FocusSmoothingInterpSpeed;
+
+	/** Additional focus depth offset, used for manually tweaking if yoour chosen focus method is needs adjustment */
+	UPROPERTY(Interp, EditAnywhere, BlueprintReadWrite, Category = "Focus Settings")
+	float FocusOffset;
+
+	FCameraFocusSettings() : 
+		FocusMethod(ECameraFocusMethod::Manual),
+		bSmoothFocusChanges(0),
+		bSimulateFocusBreathing(0),
+		ManualFocusDistance(100000.f),
+		FocusSmoothingInterpSpeed(8.f),
+		FocusOffset(0.f)
+	{}
+};
+
+/**
+ * A specialized version of a camera component, geared toward cinematic usage.
+ */
+UCLASS(
+	HideCategories = (CameraSettings), 
+	HideFunctions = (SetFieldOfView, SetAspectRatio, SetConstraintAspectRatio), 
+	Blueprintable, 
+	ClassGroup = Camera, 
+	meta = (BlueprintSpawnableComponent), 
+	Config = Engine
+	)
+class CINEMATICCAMERA_API UCineCameraComponent : public UCameraComponent
+{
+	GENERATED_BODY()
+
+public:
+	/** Default constuctor. */
+	UCineCameraComponent();
+
+	virtual void GetCameraView(float DeltaTime, FMinimalViewInfo& DesiredView) override;
+
+	UPROPERTY(EditAnywhere, BlueprintReadWrite, Category = "Current Camera Settings")
+	FCameraFilmbackSettings FilmbackSettings;
+
+	UPROPERTY(EditAnywhere, BlueprintReadWrite, Category = "Current Camera Settings")
+	FCameraLensSettings LensSettings;
+
+	UPROPERTY(EditAnywhere, BlueprintReadWrite, Category = "Current Camera Settings")
+	FCameraFocusSettings FocusSettings;
+
+	UPROPERTY(Interp, EditAnywhere, BlueprintReadWrite, Category = "Current Camera Settings")
+	float CurrentFocalLength;
+
+	UPROPERTY(Interp, EditAnywhere, BlueprintReadWrite, Category = "Current Camera Settings")
+	float CurrentAperture;
+	
+	/** Read-only. Control this value via FocusSettings. */
+	UPROPERTY(VisibleAnywhere, BlueprintReadOnly, Category = "Current Camera Settings")
+	float CurrentFocusDistance;
+
+	// for interpto
+	float LastFocusDistance;
+
+	float GetHorizontalFieldOfView() const;
+	float GetVerticalFieldOfView() const;
+
+	static TArray<FNamedFilmbackPreset> const& GetFilmbackPresets();
+	static TArray<FNamedLensPreset> const& GetLensPresets();
+
+protected:
+
+	virtual void PostLoad() override;
+#if WITH_EDITORONLY_DATA
+	virtual void PostEditChangeProperty(FPropertyChangedEvent& PropertyChangedEvent) override;
+#endif
+
+	UPROPERTY(config)
+	TArray<FNamedFilmbackPreset> FilmbackPresets;
+
+	UPROPERTY(config)
+	TArray<FNamedLensPreset> LensPresets;
+
+	virtual void UpdateCameraLens(float DeltaTime, FMinimalViewInfo& DesiredView);
+
+private:
+	void RecalcDerivedData();
+	float GetDesiredFocusDistance(FMinimalViewInfo& DesiredView) const;
+	float GetWorldToMetersScale() const;
+};
diff --git a/Engine/Source/Runtime/CinematicCamera/Public/CinematicCameraModule.h b/Engine/Source/Runtime/CinematicCamera/Public/CinematicCameraModule.h
new file mode 100644
--- /dev/null
+++ b/Engine/Source/Runtime/CinematicCamera/Public/CinematicCameraModule.h
@@ -0,0 +1,13 @@
+// Copyright 1998-2016 Epic Games, Inc. All Rights Reserved.
+
+#pragma once
+
+#include "ModuleInterface.h"
+
+/**
+ * Foliage Edit mode module interface
+ */
+class ICinematicCameraModule : public IModuleInterface
+{
+public:
+};
diff --git a/Engine/Source/Runtime/Engine/Classes/Animation/AnimMontage.h b/Engine/Source/Runtime/Engine/Classes/Animation/AnimMontage.h
--- a/Engine/Source/Runtime/Engine/Classes/Animation/AnimMontage.h
+++ b/Engine/Source/Runtime/Engine/Classes/Animation/AnimMontage.h
@@ -388,8 +388,8 @@ private:
 
 public:
 	/** static functions that are used by matinee functionality */
-	static void SetMatineeAnimPositionInner(FName SlotName, USkeletalMeshComponent* SkeletalMeshComponent, UAnimSequence* InAnimSequence, TWeakObjectPtr<UAnimMontage>& CurrentlyPlayingMontage, float InPosition, bool bLooping);
-	static void PreviewMatineeSetAnimPositionInner(FName SlotName, USkeletalMeshComponent* SkeletalMeshComponent, UAnimSequence* InAnimSequence, TWeakObjectPtr<UAnimMontage>& CurrentlyPlayingMontage, float InPosition, bool bLooping, bool bFireNotifies, float DeltaTime);
+	ENGINE_API static void SetMatineeAnimPositionInner(FName SlotName, USkeletalMeshComponent* SkeletalMeshComponent, UAnimSequence* InAnimSequence, TWeakObjectPtr<UAnimMontage>& CurrentlyPlayingMontage, float InPosition, bool bLooping);
+	ENGINE_API static void PreviewMatineeSetAnimPositionInner(FName SlotName, USkeletalMeshComponent* SkeletalMeshComponent, UAnimSequence* InAnimSequence, TWeakObjectPtr<UAnimMontage>& CurrentlyPlayingMontage, float InPosition, bool bLooping, bool bFireNotifies, float DeltaTime);
 };
 
 UCLASS(config=Engine, hidecategories=(UObject, Length), MinimalAPI, BlueprintType)
diff --git a/Engine/Source/Runtime/Engine/Classes/Animation/AnimSingleNodeInstance.h b/Engine/Source/Runtime/Engine/Classes/Animation/AnimSingleNodeInstance.h
--- a/Engine/Source/Runtime/Engine/Classes/Animation/AnimSingleNodeInstance.h
+++ b/Engine/Source/Runtime/Engine/Classes/Animation/AnimSingleNodeInstance.h
@@ -83,6 +83,8 @@ public:
 	UFUNCTION(BlueprintCallable, Category="Animation")
 	void SetPosition(float InPosition, bool bFireNotifies=true);
 	UFUNCTION(BlueprintCallable, Category="Animation")
+	void SetPositionWithPreviousTime(float InPosition, float InPreviousTime, bool bFireNotifies=true);
+	UFUNCTION(BlueprintCallable, Category="Animation")
 	void SetBlendSpaceInput(const FVector& InBlendInput);
 	UFUNCTION(BlueprintCallable, Category="Animation")
 	void SetPlaying(bool bIsPlaying);
diff --git a/Engine/Source/Runtime/Engine/Classes/Curves/CurveBase.h b/Engine/Source/Runtime/Engine/Classes/Curves/CurveBase.h
--- a/Engine/Source/Runtime/Engine/Classes/Curves/CurveBase.h
+++ b/Engine/Source/Runtime/Engine/Classes/Curves/CurveBase.h
@@ -524,7 +524,7 @@ public:
 	FRichCurveKey GetLastKey() const;
 
 	/** Get the first key that matches any of the given key handles. */
-	FRichCurveKey& GetFirstMatchingKey(const TArray<FKeyHandle>& KeyHandles);
+	FRichCurveKey* GetFirstMatchingKey(const TArray<FKeyHandle>& KeyHandles);
 
 	/** Get the next or previous key given the key handle */
 	FKeyHandle GetNextKey(FKeyHandle KeyHandle) const;
@@ -612,8 +612,13 @@ public:
 	/** Determine if two RichCurves are the same */
 	bool operator == (const FRichCurve& Curve) const;
 
+	/** Bake curve given the sample rate */
+	void BakeCurve(float SampleRate);
+	void BakeCurve(float SampleRate, float FirstKeyTime, float LastKeyTime);
+
 	/** Remove redundant keys, comparing against Tolerance */
 	void RemoveRedundantKeys(float Tolerance);
+	void RemoveRedundantKeys(float Tolerance, float FirstKeyTime, float LastKeyTime);
 
 public:
 
diff --git a/Engine/Source/Runtime/Engine/Private/Animation/AnimMontage.cpp b/Engine/Source/Runtime/Engine/Private/Animation/AnimMontage.cpp
--- a/Engine/Source/Runtime/Engine/Private/Animation/AnimMontage.cpp
+++ b/Engine/Source/Runtime/Engine/Private/Animation/AnimMontage.cpp
@@ -2042,10 +2042,21 @@ void FAnimMontageInstance::PreviewMatineeSetAnimPositionInner(FName SlotName, US
 			SingleNodeInst->SetAnimationAsset(InAnimSequence, bLooping);
 		}
 
+
 		SingleNodeInst->SetLooping(bLooping);
-		SingleNodeInst->SetPosition(InPosition, bFireNotifies);
+
+		// Anim notifies are fired from the previous evaluation time to the current time. When the delta time is 0, explicitly set the previous time to the current time so that anim notifies prior are not fired.
+		if (DeltaTime == 0.f)
+		{
+			const float PreviousTime = InPosition;
+			SingleNodeInst->SetPositionWithPreviousTime(InPosition, PreviousTime, bFireNotifies);
+		}
+		else
+		{
+			SingleNodeInst->SetPosition(InPosition, bFireNotifies);
+		}
 	}
-	else
+	else if (AnimInst)
 	{
 		bool bShouldChange = AnimInst->IsPlayingSlotAnimation(InAnimSequence, SlotName) == false;
 		if(bShouldChange)
diff --git a/Engine/Source/Runtime/Engine/Private/Animation/AnimSingleNodeInstance.cpp b/Engine/Source/Runtime/Engine/Private/Animation/AnimSingleNodeInstance.cpp
--- a/Engine/Source/Runtime/Engine/Private/Animation/AnimSingleNodeInstance.cpp
+++ b/Engine/Source/Runtime/Engine/Private/Animation/AnimSingleNodeInstance.cpp
@@ -272,11 +272,10 @@ void UAnimSingleNodeInstance::SetReverse(bool bInReverse)
 	GetProxyOnGameThread<FAnimSingleNodeInstanceProxy>().SetReverse(bInReverse);
 }
 
-void UAnimSingleNodeInstance::SetPosition(float InPosition, bool bFireNotifies)
+void UAnimSingleNodeInstance::SetPositionWithPreviousTime(float InPosition, float InPreviousTime, bool bFireNotifies)
 {
 	FAnimSingleNodeInstanceProxy& Proxy = GetProxyOnGameThread<FAnimSingleNodeInstanceProxy>();
 
-	float PreviousTime = Proxy.GetCurrentTime();
 	Proxy.SetCurrentTime(FMath::Clamp<float>(InPosition, 0.f, GetLength()));
 
 	if (FAnimMontageInstance* CurMontageInstance = GetActiveMontageInstance())
@@ -295,7 +294,7 @@ void UAnimSingleNodeInstance::SetPosition(float InPosition, bool bFireNotifies)
 			NotifyQueue.Reset(GetSkelMeshComponent());
 
 			TArray<const FAnimNotifyEvent*> Notifies;
-			SequenceBase->GetAnimNotifiesFromDeltaPositions(PreviousTime, Proxy.GetCurrentTime(), Notifies);
+			SequenceBase->GetAnimNotifiesFromDeltaPositions(InPreviousTime, Proxy.GetCurrentTime(), Notifies);
 			if ( Notifies.Num() > 0 )
 			{
 				// single node instance only has 1 asset at a time
@@ -308,6 +307,15 @@ void UAnimSingleNodeInstance::SetPosition(float InPosition, bool bFireNotifies)
 	}
 }
 
+void UAnimSingleNodeInstance::SetPosition(float InPosition, bool bFireNotifies)
+{
+	FAnimSingleNodeInstanceProxy& Proxy = GetProxyOnGameThread<FAnimSingleNodeInstanceProxy>();
+
+	float PreviousTime = Proxy.GetCurrentTime();
+
+	SetPositionWithPreviousTime(InPosition, PreviousTime, bFireNotifies);
+}
+
 void UAnimSingleNodeInstance::SetBlendSpaceInput(const FVector& InBlendInput)
 {
 	GetProxyOnGameThread<FAnimSingleNodeInstanceProxy>().SetBlendSpaceInput(InBlendInput);
diff --git a/Engine/Source/Runtime/Engine/Private/Components/SkeletalMeshComponent.cpp b/Engine/Source/Runtime/Engine/Private/Components/SkeletalMeshComponent.cpp
--- a/Engine/Source/Runtime/Engine/Private/Components/SkeletalMeshComponent.cpp
+++ b/Engine/Source/Runtime/Engine/Private/Components/SkeletalMeshComponent.cpp
@@ -604,7 +604,10 @@ void USkeletalMeshComponent::TickAnimation(float DeltaTime, bool bNeedsValidRoot
 
 void USkeletalMeshComponent::UpdateMaterialParameters()
 {
-	AnimScriptInstance->UpdateComponentsMaterialParameters(this);
+	if(AnimScriptInstance != nullptr)
+	{
+		AnimScriptInstance->UpdateComponentsMaterialParameters(this);
+	}
 }
 
 bool USkeletalMeshComponent::UpdateLODStatus()
diff --git a/Engine/Source/Runtime/Engine/Private/CurveBase.cpp b/Engine/Source/Runtime/Engine/Private/CurveBase.cpp
--- a/Engine/Source/Runtime/Engine/Private/CurveBase.cpp
+++ b/Engine/Source/Runtime/Engine/Private/CurveBase.cpp
@@ -572,18 +572,17 @@ FRichCurveKey FRichCurve::GetLastKey() const
 }
 
 
-FRichCurveKey& FRichCurve::GetFirstMatchingKey(const TArray<FKeyHandle>& KeyHandles)
+FRichCurveKey* FRichCurve::GetFirstMatchingKey(const TArray<FKeyHandle>& KeyHandles)
 {
 	for (const auto& KeyHandle : KeyHandles)
 	{
 		if (IsKeyHandleValid(KeyHandle))
 		{
-			return GetKey(KeyHandle);
+			return &GetKey(KeyHandle);
 		}
 	}
 
-	checkf(false, TEXT("No matching rich curve key found."));
-	return Keys[INDEX_NONE];
+	return nullptr;
 }
 
 
@@ -1204,12 +1203,57 @@ void FRichCurve::ReadjustTimeRange(float NewMinTimeRange, float NewMaxTimeRange,
 	}
 }
 
+void FRichCurve::BakeCurve(float SampleRate)
+{
+	if (Keys.Num() == 0)
+	{
+		return;
+	}
+
+	float FirstKeyTime = Keys[0].Time;
+	float LastKeyTime = Keys[Keys.Num()-1].Time;
+
+	BakeCurve(SampleRate, FirstKeyTime, LastKeyTime);
+}
+
+void FRichCurve::BakeCurve(float SampleRate, float FirstKeyTime, float LastKeyTime)
+{
+	if (Keys.Num() == 0)
+	{
+		return;
+	}
+
+	for (float Time = FirstKeyTime + SampleRate; Time < LastKeyTime;  )
+	{
+		float Value = Eval(Time);
+		UpdateOrAddKey(Time, Value);
+		Time += SampleRate;
+	}
+}
+
 void FRichCurve::RemoveRedundantKeys(float Tolerance)
+{
+	if (Keys.Num() == 0)
+	{
+		return;
+	}
+
+	float FirstKeyTime = Keys[0].Time;
+	float LastKeyTime = Keys[Keys.Num()-1].Time;
+
+	RemoveRedundantKeys(Tolerance, FirstKeyTime, LastKeyTime);
+}
+
+void FRichCurve::RemoveRedundantKeys(float Tolerance, float FirstKeyTime, float LastKeyTime)
 {
 	for(int32 KeyIndex = 0; KeyIndex < Keys.Num(); ++KeyIndex)
 	{
 		// copy the key
 		FRichCurveKey OriginalKey = Keys[KeyIndex];
+		if (OriginalKey.Time < FirstKeyTime || OriginalKey.Time > LastKeyTime)
+		{
+			continue;
+		}
 
 		FKeyHandle KeyHandle = GetKeyHandle(KeyIndex);
 
diff --git a/Engine/Source/Runtime/LevelSequence/Private/LevelSequence.cpp b/Engine/Source/Runtime/LevelSequence/Private/LevelSequence.cpp
--- a/Engine/Source/Runtime/LevelSequence/Private/LevelSequence.cpp
+++ b/Engine/Source/Runtime/LevelSequence/Private/LevelSequence.cpp
@@ -68,6 +68,11 @@ void ULevelSequence::BindPossessableObject(const FGuid& ObjectId, UObject& Posse
 	}
 }
 
+void ULevelSequence::BindPossessableObject(const FGuid& ObjectId, const FLevelSequenceObjectReference& ObjectReference)
+{
+	ObjectReferences.CreateBinding(ObjectId, ObjectReference);
+}
+
 bool ULevelSequence::CanPossessObject(UObject& Object) const
 {
 	return Object.IsA<AActor>() || Object.IsA<UActorComponent>();
diff --git a/Engine/Source/Runtime/LevelSequence/Private/LevelSequenceObjectReference.cpp b/Engine/Source/Runtime/LevelSequence/Private/LevelSequenceObjectReference.cpp
--- a/Engine/Source/Runtime/LevelSequence/Private/LevelSequenceObjectReference.cpp
+++ b/Engine/Source/Runtime/LevelSequence/Private/LevelSequenceObjectReference.cpp
@@ -22,6 +22,12 @@ FLevelSequenceObjectReference::FLevelSequenceObjectReference(UObject* InObject,
 	}
 }
 
+FLevelSequenceObjectReference::FLevelSequenceObjectReference(const FUniqueObjectGuid& InObjectId, const FString& InObjectPath)
+	: ObjectId(InObjectId)
+	, ObjectPath(InObjectPath)
+{
+}
+
 UObject* FLevelSequenceObjectReference::Resolve(UObject* InContext) const
 {
 	if (ObjectId.IsValid())
@@ -99,6 +105,14 @@ void FLevelSequenceObjectReferenceMap::CreateBinding(const FGuid& ObjectId, UObj
 	}
 }
 
+void FLevelSequenceObjectReferenceMap::CreateBinding(const FGuid& ObjectId, const FLevelSequenceObjectReference& ObjectReference)
+{
+	if (ensureMsgf(ObjectReference.IsValid(), TEXT("Invalid object reference specifed for binding")))
+	{
+		Map.FindOrAdd(ObjectId) = ObjectReference;
+	}
+}
+
 void FLevelSequenceObjectReferenceMap::RemoveBinding(const FGuid& ObjectId)
 {
 	Map.Remove(ObjectId);
diff --git a/Engine/Source/Runtime/LevelSequence/Private/LevelSequencePlayer.cpp b/Engine/Source/Runtime/LevelSequence/Private/LevelSequencePlayer.cpp
--- a/Engine/Source/Runtime/LevelSequence/Private/LevelSequencePlayer.cpp
+++ b/Engine/Source/Runtime/LevelSequence/Private/LevelSequencePlayer.cpp
@@ -298,6 +298,10 @@ EMovieScenePlayerStatus::Type ULevelSequencePlayer::GetPlaybackStatus() const
 	return bIsPlaying ? EMovieScenePlayerStatus::Playing : EMovieScenePlayerStatus::Stopped;
 }
 
+void ULevelSequencePlayer::SetPlaybackStatus(EMovieScenePlayerStatus::Type InPlaybackStatus)
+{
+}
+
 void ULevelSequencePlayer::AddOrUpdateMovieSceneInstance(UMovieSceneSection& MovieSceneSection, TSharedRef<FMovieSceneSequenceInstance> InstanceToAdd)
 {
 }
diff --git a/Engine/Source/Runtime/LevelSequence/Private/LevelSequenceSpawnRegister.cpp b/Engine/Source/Runtime/LevelSequence/Private/LevelSequenceSpawnRegister.cpp
--- a/Engine/Source/Runtime/LevelSequence/Private/LevelSequenceSpawnRegister.cpp
+++ b/Engine/Source/Runtime/LevelSequence/Private/LevelSequenceSpawnRegister.cpp
@@ -7,6 +7,7 @@
 #include "MovieSceneSequence.h"
 #include "MovieScene.h"
 #include "MovieSceneCommonHelpers.h"
+#include "Particles/ParticleSystem.h"
 
 UObject* FLevelSequenceSpawnRegister::SpawnObject(const FGuid& BindingId, FMovieSceneSequenceInstance& SequenceInstance, IMovieScenePlayer& Player)
 {
@@ -78,6 +79,14 @@ UObject* FLevelSequenceSpawnRegister::SpawnObject(const FGuid& BindingId, FMovie
 		{
 			TemplateRootComponent->UpdateComponentToWorld();
 			NewTransform = TemplateRootComponent->GetComponentToWorld().Inverse() * NewTransform;
+			//TemplateRootComponent->bAutoActivate = false;
+		}
+
+		// Disable all particle components so that they don't auto fire as soon as the actor is spawned. The particles should be triggered through the particle track.
+		TArray<UActorComponent*> ParticleComponents = Template->GetComponentsByClass(UParticleSystem::StaticClass());
+		for (int32 ComponentIdx = 0; ComponentIdx < ParticleComponents.Num(); ++ComponentIdx)
+		{
+			ParticleComponents[ComponentIdx]->bAutoActivate = false;
 		}
 	}
 
diff --git a/Engine/Source/Runtime/LevelSequence/Public/LevelSequence.h b/Engine/Source/Runtime/LevelSequence/Public/LevelSequence.h
--- a/Engine/Source/Runtime/LevelSequence/Public/LevelSequence.h
+++ b/Engine/Source/Runtime/LevelSequence/Public/LevelSequence.h
@@ -45,6 +45,9 @@ public:
 
 	virtual bool Rename(const TCHAR* NewName = nullptr, UObject* NewOuter = nullptr, ERenameFlags Flags = REN_None) override;
 	
+	/** Bind a posessable object with an explicitly-supplied ObjectReference */
+	void BindPossessableObject(const FGuid& ObjectId, const FLevelSequenceObjectReference& ObjectReference);
+
 private:
 
 	/** Collection of possessed objects. */
diff --git a/Engine/Source/Runtime/LevelSequence/Public/LevelSequenceObjectReference.h b/Engine/Source/Runtime/LevelSequence/Public/LevelSequenceObjectReference.h
--- a/Engine/Source/Runtime/LevelSequence/Public/LevelSequenceObjectReference.h
+++ b/Engine/Source/Runtime/LevelSequence/Public/LevelSequenceObjectReference.h
@@ -24,6 +24,14 @@ struct FLevelSequenceObjectReference
 	 */
 	LEVELSEQUENCE_API FLevelSequenceObjectReference(UObject* InObject, UObject* InContext);
 
+	/**
+	 * Generates a new reference to an object.
+	 *
+	 * @param InObjectId	Primary method of resolution - object ID, stored as an annotation on the object in the world, resolvable through TLazyObjectPtr
+	 * @param InObjectPath	Secondary method of resolution - path to the object within the context
+	 */
+	LEVELSEQUENCE_API FLevelSequenceObjectReference(const FUniqueObjectGuid& InObjectId, const FString& InObjectPath);
+
 	GENERATED_BODY();
 
 	/**
@@ -97,6 +105,14 @@ struct FLevelSequenceObjectReferenceMap
 	 */
 	void CreateBinding(const FGuid& ObjectId, UObject* InObject, UObject* InContext);
 
+	/**
+	 * Create a binding for the specified ID
+	 *
+	 * @param ObjectId	The ID to associate the object with
+	 * @param ObjectReference	The object reference to bind
+	 */
+	void CreateBinding(const FGuid& ObjectId, const FLevelSequenceObjectReference& ObjectReference);
+
 	/**
 	 * Resolve a binding for the specified ID using a given context
 	 *
