Add a new bVisualizationComponent boolean on UActorComponent. Use this instead of bIsEditor for various selection and scs tree display purposes as editor only components are entirely valid to want to see in those circumtances.
Note that components created via CreateEditorOnlyDefaultSubobject are visualization components by default for backwards compatibility purposes.

diff --git a/Engine/Source/Editor/Kismet/Private/SSCSEditor.cpp b/Engine/Source/Editor/Kismet/Private/SSCSEditor.cpp
--- a/Engine/Source/Editor/Kismet/Private/SSCSEditor.cpp
+++ b/Engine/Source/Editor/Kismet/Private/SSCSEditor.cpp
@@ -4695,7 +4695,7 @@ bool SSCSEditor::ShouldAddInstancedActorComponent(UActorComponent* ActorComp, US
 {
 	// Exclude nested DSOs attached to BP-constructed instances, which are not mutable.
 	return (ActorComp != nullptr
-		&& (ActorComp->CreationMethod == EComponentCreationMethod::Instance || !ActorComp->IsEditorOnly())
+		&& (!ActorComp->IsVisualizationComponent())
 		&& (ActorComp->CreationMethod != EComponentCreationMethod::UserConstructionScript || !GetDefault<UBlueprintEditorSettings>()->bHideConstructionScriptComponentsInDetailsView)
 		&& (ParentSceneComp == nullptr || !ParentSceneComp->IsCreatedByConstructionScript() || !ActorComp->HasAnyFlags(RF_DefaultSubObject)));
 }
diff --git a/Engine/Source/Editor/UnrealEd/Private/EditorSelectUtils.cpp b/Engine/Source/Editor/UnrealEd/Private/EditorSelectUtils.cpp
--- a/Engine/Source/Editor/UnrealEd/Private/EditorSelectUtils.cpp
+++ b/Engine/Source/Editor/UnrealEd/Private/EditorSelectUtils.cpp
@@ -724,7 +724,7 @@ bool UUnrealEdEngine::IsComponentSelected(const UPrimitiveComponent* PrimCompone
 		}
 		else
 		{
-			PotentiallySelectedComponent = (PrimComponent->IsEditorOnly() ? PrimComponent->GetAttachParent() : PrimComponent);
+			PotentiallySelectedComponent = (PrimComponent->IsVisualizationComponent() ? PrimComponent->GetAttachParent() : PrimComponent);
 		}
 
 		bIsSelected = GetSelectedComponents()->IsSelected(PotentiallySelectedComponent);
diff --git a/Engine/Source/Editor/UnrealEd/Private/Kismet2/ComponentEditorUtils.cpp b/Engine/Source/Editor/UnrealEd/Private/Kismet2/ComponentEditorUtils.cpp
--- a/Engine/Source/Editor/UnrealEd/Private/Kismet2/ComponentEditorUtils.cpp
+++ b/Engine/Source/Editor/UnrealEd/Private/Kismet2/ComponentEditorUtils.cpp
@@ -552,13 +552,18 @@ int32 FComponentEditorUtils::DeleteComponents(const TArray<UActorComponent*>& Co
 				else
 				{
 					// For a non-scene component, try to select the preceding non-scene component
-					TInlineComponentArray<UActorComponent*> ActorComponents;
-					Owner->GetComponents(ActorComponents);
-					for (int32 i = 0; i < ActorComponents.Num() && ComponentToDelete != ActorComponents[i]; ++i)
+					for (UActorComponent* Component : Owner->GetComponents())
 					{
-						if (!ActorComponents[i]->IsA(USceneComponent::StaticClass()))
+						if (Component == nullptr)
 						{
-							OutComponentToSelect = ActorComponents[i];
+							if (Component == ComponentToDelete)
+							{
+								break;
+							}
+							else if (!Component->IsA<USceneComponent>())
+							{
+								OutComponentToSelect = Component;
+							}
 						}
 					}
 				}
@@ -623,7 +628,7 @@ UActorComponent* FComponentEditorUtils::DuplicateComponent(UActorComponent* Temp
 
 	UActorComponent* NewCloneComponent = nullptr;
 	AActor* Actor = TemplateComponent->GetOwner();
-	if (!TemplateComponent->IsEditorOnly() && Actor)
+	if (!TemplateComponent->IsVisualizationComponent() && Actor)
 	{
 		Actor->Modify();
 		UClass* ComponentClass = TemplateComponent->GetClass();
diff --git a/Engine/Source/Editor/UnrealEd/Private/Kismet2/Kismet2.cpp b/Engine/Source/Editor/UnrealEd/Private/Kismet2/Kismet2.cpp
--- a/Engine/Source/Editor/UnrealEd/Private/Kismet2/Kismet2.cpp
+++ b/Engine/Source/Editor/UnrealEd/Private/Kismet2/Kismet2.cpp
@@ -1682,7 +1682,7 @@ public:
 			{
 				for (UActorComponent* ComponentToConsider : Actor->GetComponents())
 				{
-					if (ComponentToConsider && !ComponentToConsider->IsEditorOnly())
+					if (ComponentToConsider && !ComponentToConsider->IsVisualizationComponent())
 					{
 						AllSelectedComponents.Add(ComponentToConsider);
 					}
diff --git a/Engine/Source/Editor/UnrealEd/Private/LevelViewportClickHandlers.cpp b/Engine/Source/Editor/UnrealEd/Private/LevelViewportClickHandlers.cpp
--- a/Engine/Source/Editor/UnrealEd/Private/LevelViewportClickHandlers.cpp
+++ b/Engine/Source/Editor/UnrealEd/Private/LevelViewportClickHandlers.cpp
@@ -312,8 +312,8 @@ namespace ClickHandlers
 			}
 		}
 
-		//If the component selected is editor-only, we want to select the non-editor-only component it's attached to
-		while (Component != nullptr && Component->IsEditorOnly())
+		//If the component selected is a visualization component, we want to select the non-visualization component it's attached to
+		while (Component != nullptr && Component->IsVisualizationComponent())
 		{
 			Component = Component->GetAttachParent();
 		}
diff --git a/Engine/Source/Runtime/CinematicCamera/Private/CameraRig_Crane.cpp b/Engine/Source/Runtime/CinematicCamera/Private/CameraRig_Crane.cpp
--- a/Engine/Source/Runtime/CinematicCamera/Private/CameraRig_Crane.cpp
+++ b/Engine/Source/Runtime/CinematicCamera/Private/CameraRig_Crane.cpp
@@ -47,7 +47,7 @@ ACameraRig_Crane::ACameraRig_Crane(const FObjectInitializer& ObjectInitializer)
 		if (PreviewMesh_CraneBase)
 		{
 			PreviewMesh_CraneBase->SetStaticMesh(CraneBaseMesh.Object);
-			PreviewMesh_CraneBase->bIsEditorOnly = true;
+			PreviewMesh_CraneBase->SetIsVisualizationComponent(true);
 			PreviewMesh_CraneBase->SetCollisionProfileName(UCollisionProfile::NoCollision_ProfileName);
 			PreviewMesh_CraneBase->bHiddenInGame = true;
 			PreviewMesh_CraneBase->CastShadow = false;
@@ -61,7 +61,7 @@ ACameraRig_Crane::ACameraRig_Crane(const FObjectInitializer& ObjectInitializer)
 		if (PreviewMesh_CraneArm)
 		{
 			PreviewMesh_CraneArm->SetStaticMesh(CraneArmMesh.Object);
-			PreviewMesh_CraneArm->bIsEditorOnly = true;
+			PreviewMesh_CraneArm->SetIsVisualizationComponent(true);
 			PreviewMesh_CraneArm->SetCollisionProfileName(UCollisionProfile::NoCollision_ProfileName);
 			PreviewMesh_CraneArm->bHiddenInGame = true;
 			PreviewMesh_CraneArm->CastShadow = false;
@@ -78,7 +78,7 @@ ACameraRig_Crane::ACameraRig_Crane(const FObjectInitializer& ObjectInitializer)
 		if (PreviewMesh_CraneMount)
 		{
 			PreviewMesh_CraneMount->SetStaticMesh(CraneArmMount.Object);
-			PreviewMesh_CraneMount->bIsEditorOnly = true;
+			PreviewMesh_CraneMount->SetIsVisualizationComponent(true);
 			PreviewMesh_CraneMount->SetCollisionProfileName(UCollisionProfile::NoCollision_ProfileName);
 			PreviewMesh_CraneMount->bHiddenInGame = true;
 			PreviewMesh_CraneMount->CastShadow = false;
@@ -95,7 +95,7 @@ ACameraRig_Crane::ACameraRig_Crane(const FObjectInitializer& ObjectInitializer)
 		if (PreviewMesh_CraneCounterWeight)
 		{
 			PreviewMesh_CraneCounterWeight->SetStaticMesh(CraneArmCW.Object);
-			PreviewMesh_CraneCounterWeight->bIsEditorOnly = true;
+			PreviewMesh_CraneCounterWeight->SetIsVisualizationComponent(true);
 			PreviewMesh_CraneCounterWeight->SetCollisionProfileName(UCollisionProfile::NoCollision_ProfileName);
 			PreviewMesh_CraneCounterWeight->bHiddenInGame = true;
 			PreviewMesh_CraneCounterWeight->CastShadow = false;
diff --git a/Engine/Source/Runtime/CinematicCamera/Private/CameraRig_Rail.cpp b/Engine/Source/Runtime/CinematicCamera/Private/CameraRig_Rail.cpp
--- a/Engine/Source/Runtime/CinematicCamera/Private/CameraRig_Rail.cpp
+++ b/Engine/Source/Runtime/CinematicCamera/Private/CameraRig_Rail.cpp
@@ -42,7 +42,7 @@ ACameraRig_Rail::ACameraRig_Rail(const FObjectInitializer& ObjectInitializer)
 		if (PreviewMesh_Mount)
 		{
 			PreviewMesh_Mount->SetStaticMesh(MountMesh.Object);
-			PreviewMesh_Mount->bIsEditorOnly = true;
+			PreviewMesh_Mount->SetIsVisualizationComponent(true);
 			PreviewMesh_Mount->SetCollisionProfileName(UCollisionProfile::NoCollision_ProfileName);
 			PreviewMesh_Mount->bHiddenInGame = true;
 			PreviewMesh_Mount->CastShadow = false;
@@ -62,7 +62,7 @@ USplineMeshComponent* ACameraRig_Rail::CreateSplinePreviewSegment()
 	{
 		Segment->SetStaticMesh(PreviewRailStaticMesh);
 		Segment->SetMobility(EComponentMobility::Movable);
-		Segment->bIsEditorOnly = true;
+		Segment->SetIsVisualizationComponent(true);
 		Segment->SetCollisionProfileName(UCollisionProfile::NoCollision_ProfileName);
 		Segment->bHiddenInGame = true;
 		Segment->CastShadow = false;
@@ -128,7 +128,7 @@ void ACameraRig_Rail::UpdatePreviewMeshes()
 			// Unregister any owned components that aren't in PreviewRailMeshSegments
 			TArray<USplineMeshComponent*> OwnedSplineMeshComponents;
 			GetComponents(OwnedSplineMeshComponents);
-			for (auto OwnedComponent : OwnedSplineMeshComponents)
+			for (USplineMeshComponent* OwnedComponent : OwnedSplineMeshComponents)
 			{
 				if (!PreviewRailMeshSegments.Contains(OwnedComponent) && OwnedComponent->IsRegistered())
 				{
diff --git a/Engine/Source/Runtime/CinematicCamera/Private/CineCameraComponent.cpp b/Engine/Source/Runtime/CinematicCamera/Private/CineCameraComponent.cpp
--- a/Engine/Source/Runtime/CinematicCamera/Private/CineCameraComponent.cpp
+++ b/Engine/Source/Runtime/CinematicCamera/Private/CineCameraComponent.cpp
@@ -419,7 +419,7 @@ void UCineCameraComponent::CreateDebugFocusPlane()
 		{
 			DebugFocusPlaneComponent = NewObject<UStaticMeshComponent>(MyOwner, NAME_None, RF_Transactional | RF_TextExportTransient);
 			DebugFocusPlaneComponent->SetupAttachment(this);
-			DebugFocusPlaneComponent->bIsEditorOnly = false;
+			DebugFocusPlaneComponent->SetIsVisualizationComponent(true);
 			DebugFocusPlaneComponent->SetStaticMesh(FocusPlaneVisualizationMesh);
 			DebugFocusPlaneComponent->SetCollisionProfileName(UCollisionProfile::NoCollision_ProfileName);
 			DebugFocusPlaneComponent->bHiddenInGame = false;
diff --git a/Engine/Source/Runtime/Engine/Classes/Components/ActorComponent.h b/Engine/Source/Runtime/Engine/Classes/Components/ActorComponent.h
--- a/Engine/Source/Runtime/Engine/Classes/Components/ActorComponent.h
+++ b/Engine/Source/Runtime/Engine/Classes/Components/ActorComponent.h
@@ -210,10 +210,16 @@ public:
 	DEPRECATED(4.14, "bWantsBeginPlay was inconsistently enforced and is now unused. BeginPlay will now always be called for Actor Components.")
 	uint8 bWantsBeginPlay:1;
 
+#if WITH_EDITORONLY_DATA
 	/** If true, the component will be excluded from non-editor builds */
 	UPROPERTY(EditAnywhere, BlueprintReadOnly, Category = Cooking)
 	uint8 bIsEditorOnly:1;
 
+private:
+	UPROPERTY()
+	uint8 bIsVisualizationComponent : 1;
+#endif
+
 private:
 	/** Indicates that OnCreatedComponent has been called, but OnDestroyedComponent has not yet */
 	uint8 bHasBeenCreated:1;
@@ -393,14 +399,29 @@ public:
 
 	virtual bool GetComponentClassCanReplicate() const;
 
+#if WITH_EDITORONLY_DATA
 	/** Returns whether this component is an editor-only object or not */
 	virtual bool IsEditorOnly() const override { return bIsEditorOnly; }
 
 	virtual void MarkAsEditorOnlySubobject() override
 	{
 		bIsEditorOnly = true;
+		// A bit sketchy, but is best for backwards compatibility as the vast majority of editor only components were for visualization, 
+		// so for the very few where visualization is not the purpose, it can be cleared after the subobject is created
+		bIsVisualizationComponent = true; 
 	}
 
+	bool IsVisualizationComponent() const { return bIsVisualizationComponent; }
+	void SetIsVisualizationComponent(const bool bInIsVisualizationComponent)
+	{
+		bIsVisualizationComponent = bInIsVisualizationComponent;
+		if (bIsVisualizationComponent)
+		{
+			bIsEditorOnly = true;
+		}
+	}
+#endif
+
 	/** Returns net role of the owning actor */
 	/** Returns true if we are replicating and not authorative */
 	bool	IsNetSimulating() const;
diff --git a/Engine/Source/Runtime/Engine/Private/Camera/CameraComponent.cpp b/Engine/Source/Runtime/Engine/Private/Camera/CameraComponent.cpp
--- a/Engine/Source/Runtime/Engine/Private/Camera/CameraComponent.cpp
+++ b/Engine/Source/Runtime/Engine/Private/Camera/CameraComponent.cpp
@@ -93,7 +93,7 @@ void UCameraComponent::OnRegister()
 		{
 			ProxyMeshComponent = NewObject<UStaticMeshComponent>(MyOwner, NAME_None, RF_Transactional | RF_TextExportTransient);
 			ProxyMeshComponent->SetupAttachment(this);
-			ProxyMeshComponent->bIsEditorOnly = true;
+			ProxyMeshComponent->SetIsVisualizationComponent(true);
 			ProxyMeshComponent->SetStaticMesh(CameraMesh);
 			ProxyMeshComponent->SetCollisionProfileName(UCollisionProfile::NoCollision_ProfileName);
 			ProxyMeshComponent->bHiddenInGame = true;
@@ -107,7 +107,7 @@ void UCameraComponent::OnRegister()
 		{
 			DrawFrustum = NewObject<UDrawFrustumComponent>(MyOwner, NAME_None, RF_Transactional | RF_TextExportTransient);
 			DrawFrustum->SetupAttachment(this);
-			DrawFrustum->bIsEditorOnly = true;
+			DrawFrustum->SetIsVisualizationComponent(true);
 			DrawFrustum->CreationMethod = CreationMethod;
 			DrawFrustum->RegisterComponentWithWorld(GetWorld());
 		}
diff --git a/Engine/Source/Runtime/Engine/Private/Components/SceneCaptureComponent.cpp b/Engine/Source/Runtime/Engine/Private/Components/SceneCaptureComponent.cpp
--- a/Engine/Source/Runtime/Engine/Private/Components/SceneCaptureComponent.cpp
+++ b/Engine/Source/Runtime/Engine/Private/Components/SceneCaptureComponent.cpp
@@ -51,7 +51,7 @@ ASceneCapture2D::ASceneCapture2D(const FObjectInitializer& ObjectInitializer)
 	: Super(ObjectInitializer)
 {
 	DrawFrustum = CreateDefaultSubobject<UDrawFrustumComponent>(TEXT("DrawFrust0"));
-	DrawFrustum->bIsEditorOnly = true;
+	DrawFrustum->SetIsVisualizationComponent(true);
 	DrawFrustum->SetupAttachment(GetMeshComp());
 
 	CaptureComponent2D = CreateDefaultSubobject<USceneCaptureComponent2D>(TEXT("NewSceneCaptureComponent2D"));
@@ -106,7 +106,7 @@ ASceneCaptureCube::ASceneCaptureCube(const FObjectInitializer& ObjectInitializer
 	: Super(ObjectInitializer)
 {
 	DrawFrustum = CreateDefaultSubobject<UDrawFrustumComponent>(TEXT("DrawFrust0"));
-	DrawFrustum->bIsEditorOnly = true;
+	DrawFrustum->SetIsVisualizationComponent(true);
 	DrawFrustum->SetupAttachment(GetMeshComp());
 
 	CaptureComponentCube = CreateDefaultSubobject<USceneCaptureComponentCube>(TEXT("NewSceneCaptureComponentCube"));
diff --git a/Engine/Source/Runtime/Engine/Private/Components/SceneComponent.cpp b/Engine/Source/Runtime/Engine/Private/Components/SceneComponent.cpp
--- a/Engine/Source/Runtime/Engine/Private/Components/SceneComponent.cpp
+++ b/Engine/Source/Runtime/Engine/Private/Components/SceneComponent.cpp
@@ -669,7 +669,7 @@ void USceneComponent::OnRegister()
 		SpriteComponent->RelativeScale3D = FVector(0.5f, 0.5f, 0.5f);
 		SpriteComponent->Mobility = EComponentMobility::Movable;
 		SpriteComponent->AlwaysLoadOnClient = false;
-		SpriteComponent->bIsEditorOnly = true;
+		SpriteComponent->SetIsVisualizationComponent(true);
 		SpriteComponent->SpriteInfo.Category = TEXT("Misc");
 		SpriteComponent->SpriteInfo.DisplayName = NSLOCTEXT( "SpriteCategory", "Misc", "Misc" );
 		SpriteComponent->CreationMethod = CreationMethod;
