PR #4728: Description Correction (Contributed by [REDACTED_PERSON])

diff --git a/Engine/Source/Runtime/CinematicCamera/Public/CineCameraComponent.h b/Engine/Source/Runtime/CinematicCamera/Public/CineCameraComponent.h
--- a/Engine/Source/Runtime/CinematicCamera/Public/CineCameraComponent.h
+++ b/Engine/Source/Runtime/CinematicCamera/Public/CineCameraComponent.h
@@ -72,7 +72,7 @@ struct FCameraLensSettings
 	UPROPERTY(EditAnywhere, BlueprintReadWrite, Category = "Lens")
 	float MinFStop;
 
-	/** Minimum aperture for this lens (e.g. 2.8 for an f/2.8 lens) */
+	/** Maximum aperture for this lens (e.g. 2.8 for an f/2.8 lens) */
 	UPROPERTY(EditAnywhere, BlueprintReadWrite, Category = "Lens")
 	float MaxFStop;
 
