PR #4577: UE-56250: Remove debug from shipping builds (Contributed by [REDACTED_PERSON])

diff --git a/Engine/Source/Runtime/CinematicCamera/Private/CineCameraComponent.cpp b/Engine/Source/Runtime/CinematicCamera/Private/CineCameraComponent.cpp
--- a/Engine/Source/Runtime/CinematicCamera/Private/CineCameraComponent.cpp
+++ b/Engine/Source/Runtime/CinematicCamera/Private/CineCameraComponent.cpp
@@ -90,6 +90,7 @@ static const FColor DebugFocusPointOutlineColor = FColor::Black;
 
 void UCineCameraComponent::TickComponent(float DeltaTime, ELevelTick TickType, FActorComponentTickFunction* ThisTickFunction)
 {
+#if !(UE_BUILD_SHIPPING || UE_BUILD_TEST)
 	// make sure drawing is set up
 	if (FocusSettings.bDrawDebugFocusPlane)
 	{
@@ -107,6 +108,7 @@ void UCineCameraComponent::TickComponent(float DeltaTime, ELevelTick TickType, F
 			DestroyDebugFocusPlane();
 		}
 	}
+#endif // !(UE_BUILD_SHIPPING || UE_BUILD_TEST)
 
 #if ENABLE_DRAW_DEBUG
 	if (FocusSettings.TrackingFocusSettings.bDrawDebugTrackingFocusPoint)
@@ -337,6 +339,7 @@ void UCineCameraComponent::GetCameraView(float DeltaTime, FMinimalViewInfo& Desi
 	bResetInterpolation = false;
 }
 
+#if !(UE_BUILD_SHIPPING || UE_BUILD_TEST)
 void UCineCameraComponent::UpdateDebugFocusPlane()
 {
 	if (FocusPlaneVisualizationMesh && DebugFocusPlaneComponent)
@@ -347,7 +350,7 @@ void UCineCameraComponent::UpdateDebugFocusPlane()
 		DebugFocusPlaneComponent->SetWorldLocation(FocusPoint);
 	}
 }
-
+#endif // !(UE_BUILD_SHIPPING || UE_BUILD_TEST)
 
 void UCineCameraComponent::UpdateCameraLens(float DeltaTime, FMinimalViewInfo& DesiredView)
 {
@@ -408,6 +411,7 @@ void UCineCameraComponent::NotifyCameraCut()
 	bResetInterpolation = true;
 }
 
+#if !(UE_BUILD_SHIPPING || UE_BUILD_TEST)
 void UCineCameraComponent::CreateDebugFocusPlane()
 {
 	if (AActor* const MyOwner = GetOwner())
@@ -449,6 +453,7 @@ void UCineCameraComponent::DestroyDebugFocusPlane()
 		DebugFocusPlaneMID = nullptr;
 	}
 }
+#endif // !(UE_BUILD_SHIPPING || UE_BUILD_TEST)
 
 void UCineCameraComponent::OnRegister()
 {
diff --git a/Engine/Source/Runtime/CinematicCamera/Public/CineCameraComponent.h b/Engine/Source/Runtime/CinematicCamera/Public/CineCameraComponent.h
--- a/Engine/Source/Runtime/CinematicCamera/Public/CineCameraComponent.h
+++ b/Engine/Source/Runtime/CinematicCamera/Public/CineCameraComponent.h
@@ -275,8 +275,10 @@ public:
 	/** Returns a list of available lens presets. */
 	static TArray<FNamedLensPreset> const& GetLensPresets();
 
+#if !(UE_BUILD_SHIPPING || UE_BUILD_TEST)
 	/** Update the debug focus plane position and orientation. */
 	void UpdateDebugFocusPlane();
+#endif // !(UE_BUILD_SHIPPING || UE_BUILD_TEST)
 
 protected:
 
@@ -355,6 +357,8 @@ private:
 	float GetDesiredFocusDistance(const FVector& InLocation) const;
 	float GetWorldToMetersScale() const;
 
+#if !(UE_BUILD_SHIPPING || UE_BUILD_TEST)
 	void CreateDebugFocusPlane();
 	void DestroyDebugFocusPlane();
+#endif // !(UE_BUILD_SHIPPING || UE_BUILD_TEST)
 };
