Cine Camera: Change ActorToTrack to a TSoftObjectPtr<AActor> to support cross level references

diff --git a/Engine/Source/Runtime/CinematicCamera/Private/CineCameraActor.cpp b/Engine/Source/Runtime/CinematicCamera/Private/CineCameraActor.cpp
--- a/Engine/Source/Runtime/CinematicCamera/Private/CineCameraActor.cpp
+++ b/Engine/Source/Runtime/CinematicCamera/Private/CineCameraActor.cpp
@@ -34,9 +34,9 @@ bool ACineCameraActor::ShouldTickIfViewportsOnly() const
 FVector ACineCameraActor::GetLookatLocation() const
 {
 	FVector FinalLookat;
-	if (LookatTrackingSettings.ActorToTrack)
+	if (AActor* ActorToTrack = LookatTrackingSettings.ActorToTrack.Get())
 	{
-		FTransform const BaseTransform = LookatTrackingSettings.ActorToTrack->GetActorTransform();
+		FTransform const BaseTransform = ActorToTrack->GetActorTransform();
 		FinalLookat = BaseTransform.TransformPosition(LookatTrackingSettings.RelativeOffset);
 	}
 	else
diff --git a/Engine/Source/Runtime/CinematicCamera/Private/CineCameraComponent.cpp b/Engine/Source/Runtime/CinematicCamera/Private/CineCameraComponent.cpp
--- a/Engine/Source/Runtime/CinematicCamera/Private/CineCameraComponent.cpp
+++ b/Engine/Source/Runtime/CinematicCamera/Private/CineCameraComponent.cpp
@@ -113,7 +113,7 @@ void UCineCameraComponent::TickComponent(float DeltaTime, ELevelTick TickType, F
 #if ENABLE_DRAW_DEBUG
 	if (FocusSettings.TrackingFocusSettings.bDrawDebugTrackingFocusPoint)
 	{
-		AActor const* const TrackedActor = FocusSettings.TrackingFocusSettings.ActorToTrack;
+		AActor const* const TrackedActor = FocusSettings.TrackingFocusSettings.ActorToTrack.Get();
 
 		FVector FocusPoint;
 		if (TrackedActor)
@@ -304,7 +304,7 @@ float UCineCameraComponent::GetDesiredFocusDistance(const FVector& InLocation) c
 
 	case ECameraFocusMethod::Tracking:
 		{
-			AActor const* const TrackedActor = FocusSettings.TrackingFocusSettings.ActorToTrack;
+			AActor const* const TrackedActor = FocusSettings.TrackingFocusSettings.ActorToTrack.Get();
 
 			FVector FocusPoint;
 			if (TrackedActor)
diff --git a/Engine/Source/Runtime/CinematicCamera/Public/CineCameraActor.h b/Engine/Source/Runtime/CinematicCamera/Public/CineCameraActor.h
--- a/Engine/Source/Runtime/CinematicCamera/Public/CineCameraActor.h
+++ b/Engine/Source/Runtime/CinematicCamera/Public/CineCameraActor.h
@@ -33,7 +33,7 @@ struct FCameraLookatTrackingSettings
 
 	/** If set, camera will track this actor's location */
 	UPROPERTY(Interp, EditAnywhere, BlueprintReadWrite, Category = "LookAt")
-	AActor* ActorToTrack;
+	TSoftObjectPtr<AActor> ActorToTrack;
 
 	/** Offset from actor position to look at. Relative to actor if tracking an actor, relative to world otherwise. */
 	UPROPERTY(Interp, EditAnywhere, BlueprintReadWrite, Category = "LookAt")
diff --git a/Engine/Source/Runtime/CinematicCamera/Public/CineCameraComponent.h b/Engine/Source/Runtime/CinematicCamera/Public/CineCameraComponent.h
--- a/Engine/Source/Runtime/CinematicCamera/Public/CineCameraComponent.h
+++ b/Engine/Source/Runtime/CinematicCamera/Public/CineCameraComponent.h
@@ -131,7 +131,7 @@ struct FCameraTrackingFocusSettings
 
 	/** Focus distance will be tied to this actor's location. */
 	UPROPERTY(Interp, EditAnywhere, BlueprintReadWrite, Category = "Tracking Focus")
-	AActor* ActorToTrack;
+	TSoftObjectPtr<AActor> ActorToTrack;
 
 	/** Offset from actor position to track. Relative to actor if tracking an actor, relative to world otherwise. */
 	UPROPERTY(Interp, EditAnywhere, BlueprintReadWrite, Category = "Tracking Focus")
@@ -142,8 +142,7 @@ struct FCameraTrackingFocusSettings
 	uint8 bDrawDebugTrackingFocusPoint : 1;
 
 	FCameraTrackingFocusSettings()
-		: ActorToTrack(nullptr),
-		bDrawDebugTrackingFocusPoint(false)
+		: bDrawDebugTrackingFocusPoint(false)
 	{}
 };
 
