Rollback 4181586 because keying SoftObjectPtr isn't supported yet in 4.21. Cine Camera: Change ActorToTrack to a TSoftObjectPtr<AActor> to support cross level references

diff --git a/Engine/Source/Runtime/CinematicCamera/Private/CineCameraActor.cpp b/Engine/Source/Runtime/CinematicCamera/Private/CineCameraActor.cpp
--- a/Engine/Source/Runtime/CinematicCamera/Private/CineCameraActor.cpp
+++ b/Engine/Source/Runtime/CinematicCamera/Private/CineCameraActor.cpp
@@ -34,9 +34,9 @@ bool ACineCameraActor::ShouldTickIfViewportsOnly() const
 FVector ACineCameraActor::GetLookatLocation() const
 {
 	FVector FinalLookat;
-	if (AActor* ActorToTrack = LookatTrackingSettings.ActorToTrack.Get())
+	if (LookatTrackingSettings.ActorToTrack)
 	{
-		FTransform const BaseTransform = ActorToTrack->GetActorTransform();
+		FTransform const BaseTransform = LookatTrackingSettings.ActorToTrack->GetActorTransform();
 		FinalLookat = BaseTransform.TransformPosition(LookatTrackingSettings.RelativeOffset);
 	}
 	else
diff --git a/Engine/Source/Runtime/CinematicCamera/Private/CineCameraComponent.cpp b/Engine/Source/Runtime/CinematicCamera/Private/CineCameraComponent.cpp
--- a/Engine/Source/Runtime/CinematicCamera/Private/CineCameraComponent.cpp
+++ b/Engine/Source/Runtime/CinematicCamera/Private/CineCameraComponent.cpp
@@ -115,7 +115,7 @@ void UCineCameraComponent::TickComponent(float DeltaTime, ELevelTick TickType, F
 #if ENABLE_DRAW_DEBUG
 	if (FocusSettings.TrackingFocusSettings.bDrawDebugTrackingFocusPoint)
 	{
-		AActor const* const TrackedActor = FocusSettings.TrackingFocusSettings.ActorToTrack.Get();
+		AActor const* const TrackedActor = FocusSettings.TrackingFocusSettings.ActorToTrack;
 
 		FVector FocusPoint;
 		if (TrackedActor)
@@ -306,7 +306,7 @@ float UCineCameraComponent::GetDesiredFocusDistance(const FVector& InLocation) c
 
 	case ECameraFocusMethod::Tracking:
 		{
-			AActor const* const TrackedActor = FocusSettings.TrackingFocusSettings.ActorToTrack.Get();
+			AActor const* const TrackedActor = FocusSettings.TrackingFocusSettings.ActorToTrack;
 
 			FVector FocusPoint;
 			if (TrackedActor)
diff --git a/Engine/Source/Runtime/CinematicCamera/Public/CineCameraActor.h b/Engine/Source/Runtime/CinematicCamera/Public/CineCameraActor.h
--- a/Engine/Source/Runtime/CinematicCamera/Public/CineCameraActor.h
+++ b/Engine/Source/Runtime/CinematicCamera/Public/CineCameraActor.h
@@ -33,7 +33,7 @@ struct FCameraLookatTrackingSettings
 
 	/** If set, camera will track this actor's location */
 	UPROPERTY(Interp, EditAnywhere, BlueprintReadWrite, Category = "LookAt")
-	TSoftObjectPtr<AActor> ActorToTrack;
+	AActor* ActorToTrack;
 
 	/** Offset from actor position to look at. Relative to actor if tracking an actor, relative to world otherwise. */
 	UPROPERTY(Interp, EditAnywhere, BlueprintReadWrite, Category = "LookAt")
diff --git a/Engine/Source/Runtime/CinematicCamera/Public/CineCameraComponent.h b/Engine/Source/Runtime/CinematicCamera/Public/CineCameraComponent.h
--- a/Engine/Source/Runtime/CinematicCamera/Public/CineCameraComponent.h
+++ b/Engine/Source/Runtime/CinematicCamera/Public/CineCameraComponent.h
@@ -138,7 +138,7 @@ struct FCameraTrackingFocusSettings
 
 	/** Focus distance will be tied to this actor's location. */
 	UPROPERTY(Interp, EditAnywhere, BlueprintReadWrite, Category = "Tracking Focus")
-	TSoftObjectPtr<AActor> ActorToTrack;
+	AActor* ActorToTrack;
 
 	/** Offset from actor position to track. Relative to actor if tracking an actor, relative to world otherwise. */
 	UPROPERTY(Interp, EditAnywhere, BlueprintReadWrite, Category = "Tracking Focus")
