Cine Camera: Fix warning

diff --git a/Engine/Source/Runtime/CinematicCamera/Public/CineCameraComponent.h b/Engine/Source/Runtime/CinematicCamera/Public/CineCameraComponent.h
--- a/Engine/Source/Runtime/CinematicCamera/Public/CineCameraComponent.h
+++ b/Engine/Source/Runtime/CinematicCamera/Public/CineCameraComponent.h
@@ -149,7 +149,8 @@ struct FCameraTrackingFocusSettings
 	uint8 bDrawDebugTrackingFocusPoint : 1;
 
 	FCameraTrackingFocusSettings()
-		: RelativeOffset(ForceInitToZero),
+		: ActorToTrack(nullptr),
+		RelativeOffset(ForceInitToZero),
 		bDrawDebugTrackingFocusPoint(false)
 	{}
 };
