CameraRig: Preview meshes should be transient and not saved. This fixes a bug where they are being serialized out as null.

diff --git a/Engine/Source/Runtime/CinematicCamera/Public/CameraRig_Crane.h b/Engine/Source/Runtime/CinematicCamera/Public/CameraRig_Crane.h
--- a/Engine/Source/Runtime/CinematicCamera/Public/CameraRig_Crane.h
+++ b/Engine/Source/Runtime/CinematicCamera/Public/CameraRig_Crane.h
@@ -77,16 +77,16 @@ private:
 
 #if WITH_EDITORONLY_DATA
 	/** Preview meshes for visualization */
-	UPROPERTY()
+	UPROPERTY(Transient)
 	UStaticMeshComponent* PreviewMesh_CraneArm;
 
-	UPROPERTY()
+	UPROPERTY(Transient)
 	UStaticMeshComponent* PreviewMesh_CraneBase;
 
-	UPROPERTY()
+	UPROPERTY(Transient)
 	UStaticMeshComponent* PreviewMesh_CraneMount;
 
-	UPROPERTY()
+	UPROPERTY(Transient)
 	UStaticMeshComponent* PreviewMesh_CraneCounterWeight;
 #endif
 };
diff --git a/Engine/Source/Runtime/CinematicCamera/Public/CameraRig_Rail.h b/Engine/Source/Runtime/CinematicCamera/Public/CameraRig_Rail.h
--- a/Engine/Source/Runtime/CinematicCamera/Public/CameraRig_Rail.h
+++ b/Engine/Source/Runtime/CinematicCamera/Public/CameraRig_Rail.h
@@ -66,16 +66,16 @@ private:
 
 #if WITH_EDITORONLY_DATA
 	/** Preview meshes for visualization */
-	UPROPERTY()
+	UPROPERTY(Transient)
 	USplineMeshComponent* PreviewMesh_Rail;
 
-	UPROPERTY()
+	UPROPERTY(Transient)
 	TArray<USplineMeshComponent*> PreviewRailMeshSegments;
 
-	UPROPERTY()
+	UPROPERTY(Transient)
 	UStaticMesh* PreviewRailStaticMesh;
 
-	UPROPERTY()
+	UPROPERTY(Transient)
 	UStaticMeshComponent* PreviewMesh_Mount;
 #endif
 };
