CameraRig: Remove invalid preview rail mesh segments

diff --git a/Engine/Source/Runtime/CinematicCamera/Private/CameraRig_Rail.cpp b/Engine/Source/Runtime/CinematicCamera/Private/CameraRig_Rail.cpp
--- a/Engine/Source/Runtime/CinematicCamera/Private/CameraRig_Rail.cpp
+++ b/Engine/Source/Runtime/CinematicCamera/Private/CameraRig_Rail.cpp
@@ -84,6 +84,8 @@ void ACameraRig_Rail::UpdatePreviewMeshes()
 
 			// make sure our preview mesh array is correctly sized and populated
 			{
+				PreviewRailMeshSegments.RemoveAll([](USplineMeshComponent* Segment) { return Segment == nullptr; });
+
 				int32 const NumExistingPreviewMeshes = PreviewRailMeshSegments.Num();
 				if (NumExistingPreviewMeshes > NumNeededPreviewMeshes)
 				{
