Combine transforms for CameraRig_Crane to cut down on component move time

diff --git a/Engine/Source/Runtime/CinematicCamera/Private/CameraRig_Crane.cpp b/Engine/Source/Runtime/CinematicCamera/Private/CameraRig_Crane.cpp
--- a/Engine/Source/Runtime/CinematicCamera/Private/CameraRig_Crane.cpp
+++ b/Engine/Source/Runtime/CinematicCamera/Private/CameraRig_Crane.cpp
@@ -145,7 +145,6 @@ void ACameraRig_Crane::UpdateCraneComponents()
 
 	FVector NewCameraMountLoc = CraneCameraMount->RelativeLocation;
 	NewCameraMountLoc.X = CraneArmLength;
-	CraneCameraMount->SetRelativeLocation(NewCameraMountLoc);
 
 	// zero the pitch from the camera mount component
 	// this effectively gives us bAbsoluteRotation for only pitch component of an attached camera
@@ -159,7 +158,9 @@ void ACameraRig_Crane::UpdateCraneComponents()
 		NewCameraMountWorldRot.Yaw = RootComponent->RelativeRotation.Yaw;
 	}
 	NewCameraMountWorldRot.Roll = 0.f;
-	CraneCameraMount->SetWorldRotation(NewCameraMountWorldRot);
+
+	FQuat RelativeRot = CraneCameraMount->GetRelativeRotationFromWorld(FQuat(NewCameraMountWorldRot));
+	CraneCameraMount->SetRelativeTransform(FTransform(RelativeRot, NewCameraMountLoc));
 
 #if WITH_EDITORONLY_DATA
 	UpdatePreviewMeshes();
