Combine transforms for CameraRig_Rail update to cut down on component move time

diff --git a/Engine/Source/Runtime/CinematicCamera/Private/CameraRig_Rail.cpp b/Engine/Source/Runtime/CinematicCamera/Private/CameraRig_Rail.cpp
--- a/Engine/Source/Runtime/CinematicCamera/Private/CameraRig_Rail.cpp
+++ b/Engine/Source/Runtime/CinematicCamera/Private/CameraRig_Rail.cpp
@@ -163,12 +163,15 @@ void ACameraRig_Rail::UpdateRailComponents()
 	{
 		float const SplineLen = RailSplineComponent->GetSplineLength();
 		FVector const MountPos = RailSplineComponent->GetLocationAtDistanceAlongSpline(CurrentPositionOnRail*SplineLen, ESplineCoordinateSpace::World);
-		RailCameraMount->SetWorldLocation(MountPos);
 
 		if (bLockOrientationToRail)
 		{
 			FQuat const RailRot = RailSplineComponent->GetQuaternionAtDistanceAlongSpline(CurrentPositionOnRail*SplineLen, ESplineCoordinateSpace::World);
-			RailCameraMount->SetWorldRotation(RailRot);
+			RailCameraMount->SetWorldTransform(FTransform(RailRot, MountPos));
+		}
+		else
+		{
+			RailCameraMount->SetWorldLocation(MountPos);
 		}
 	}
 
