Camera Rig: Fix spline mesh preview scale

diff --git a/Engine/Source/Runtime/CinematicCamera/Private/CameraRig_Rail.cpp b/Engine/Source/Runtime/CinematicCamera/Private/CameraRig_Rail.cpp
--- a/Engine/Source/Runtime/CinematicCamera/Private/CameraRig_Rail.cpp
+++ b/Engine/Source/Runtime/CinematicCamera/Private/CameraRig_Rail.cpp
@@ -126,7 +126,8 @@ void ACameraRig_Rail::UpdatePreviewMeshes()
 				if (SplineMeshComp)
 				{
 					SplineMeshComp->SetVisibility(bShowRailVisualization);
-					SplineMeshComp->SetWorldScale3D(FVector(1.f, PreviewMeshScale, PreviewMeshScale));
+					SplineMeshComp->SetStartScale(FVector2D(PreviewMeshScale, PreviewMeshScale));
+					SplineMeshComp->SetEndScale(FVector2D(PreviewMeshScale, PreviewMeshScale));
 					SplineMeshComp->SetForwardAxis(ESplineMeshAxis::Z);
 					SplineMeshComp->SetStartAndEnd(StartLoc, StartTangent, EndLoc, EndTangent, true);
 				}
