Cine Camera: Deprecate FilmbackSettings in favor of Filmback

diff --git a/Engine/Plugins/Animation/LiveLink/Source/LiveLinkComponents/Private/Controllers/LiveLinkCameraController.cpp b/Engine/Plugins/Animation/LiveLink/Source/LiveLinkComponents/Private/Controllers/LiveLinkCameraController.cpp
--- a/Engine/Plugins/Animation/LiveLink/Source/LiveLinkComponents/Private/Controllers/LiveLinkCameraController.cpp
+++ b/Engine/Plugins/Animation/LiveLink/Source/LiveLinkComponents/Private/Controllers/LiveLinkCameraController.cpp
@@ -47,8 +47,8 @@ void ULiveLinkCameraController::Tick(float DeltaTime, const FLiveLinkSubjectRepr
 				{
 					if (StaticData->bIsFocalLengthSupported) { CineCameraComponent->CurrentFocalLength = FrameData->FocalLength; }
 					if (StaticData->bIsApertureSupported) { CineCameraComponent->CurrentAperture = FrameData->Aperture; }
-					if (StaticData->FilmBackWidth > 0.0f) { CineCameraComponent->FilmbackSettings.SensorWidth = StaticData->FilmBackWidth; }
-					if (StaticData->FilmBackHeight > 0.0f) { CineCameraComponent->FilmbackSettings.SensorHeight = StaticData->FilmBackHeight; }
+					if (StaticData->FilmBackWidth > 0.0f) { CineCameraComponent->Filmback.SensorWidth = StaticData->FilmBackWidth; }
+					if (StaticData->FilmBackHeight > 0.0f) { CineCameraComponent->Filmback.SensorHeight = StaticData->FilmBackHeight; }
 					if (StaticData->bIsFocusDistanceSupported) { CineCameraComponent->FocusSettings.ManualFocusDistance = FrameData->FocusDistance; }
 				}
 			}
diff --git a/Engine/Plugins/Enterprise/DatasmithContent/Source/DatasmithContent/Private/ObjectTemplates/DatasmithCineCameraComponentTemplate.cpp b/Engine/Plugins/Enterprise/DatasmithContent/Source/DatasmithContent/Private/ObjectTemplates/DatasmithCineCameraComponentTemplate.cpp
--- a/Engine/Plugins/Enterprise/DatasmithContent/Source/DatasmithContent/Private/ObjectTemplates/DatasmithCineCameraComponentTemplate.cpp
+++ b/Engine/Plugins/Enterprise/DatasmithContent/Source/DatasmithContent/Private/ObjectTemplates/DatasmithCineCameraComponentTemplate.cpp
@@ -162,7 +162,7 @@ UObject* UDatasmithCineCameraComponentTemplate::UpdateObject( UObject* Destinati
 	DATASMITHOBJECTTEMPLATE_CONDITIONALSET( CurrentFocalLength, CineCameraComponent, PreviousTemplate );
 	DATASMITHOBJECTTEMPLATE_CONDITIONALSET( CurrentAperture, CineCameraComponent, PreviousTemplate );
 
-	FilmbackSettings.Apply( &CineCameraComponent->FilmbackSettings, PreviousTemplate ? &PreviousTemplate->FilmbackSettings : nullptr );
+	FilmbackSettings.Apply( &CineCameraComponent->Filmback, PreviousTemplate ? &PreviousTemplate->FilmbackSettings : nullptr );
 	LensSettings.Apply( &CineCameraComponent->LensSettings, PreviousTemplate ? &PreviousTemplate->LensSettings : nullptr );
 	FocusSettings.Apply( &CineCameraComponent->FocusSettings, PreviousTemplate ? &PreviousTemplate->FocusSettings : nullptr );
 
@@ -185,7 +185,7 @@ void UDatasmithCineCameraComponentTemplate::Load( const UObject* Source )
 	CurrentFocalLength = CineCameraComponent->CurrentFocalLength;
 	CurrentAperture = CineCameraComponent->CurrentAperture;
 
-	FilmbackSettings.Load( CineCameraComponent->FilmbackSettings );
+	FilmbackSettings.Load( CineCameraComponent->Filmback );
 	LensSettings.Load( CineCameraComponent->LensSettings );
 	FocusSettings.Load( CineCameraComponent->FocusSettings );
 
diff --git a/Engine/Plugins/Enterprise/DatasmithImporter/Source/DatasmithImporter/Private/Utility/DatasmithImporterUtils.cpp b/Engine/Plugins/Enterprise/DatasmithImporter/Source/DatasmithImporter/Private/Utility/DatasmithImporterUtils.cpp
--- a/Engine/Plugins/Enterprise/DatasmithImporter/Source/DatasmithImporter/Private/Utility/DatasmithImporterUtils.cpp
+++ b/Engine/Plugins/Enterprise/DatasmithImporter/Source/DatasmithImporter/Private/Utility/DatasmithImporterUtils.cpp
@@ -644,8 +644,8 @@ namespace FDatasmithImporterUtilsHelper
 
 		const UCineCameraComponent* CineCameraComponent = CameraActor->GetCineCameraComponent();
 
-		CameraElement->SetSensorWidth( CineCameraComponent->FilmbackSettings.SensorWidth );
-		CameraElement->SetSensorAspectRatio( CineCameraComponent->FilmbackSettings.SensorWidth / CineCameraComponent->FilmbackSettings.SensorHeight );
+		CameraElement->SetSensorWidth( CineCameraComponent->Filmback.SensorWidth );
+		CameraElement->SetSensorAspectRatio( CineCameraComponent->Filmback.SensorWidth / CineCameraComponent->Filmback.SensorHeight );
 		CameraElement->SetFocalLength( CineCameraComponent->CurrentFocalLength );
 		CameraElement->SetFStop( CineCameraComponent->CurrentAperture );
 		CameraElement->SetEnableDepthOfField( CineCameraComponent->FocusSettings.FocusMethod == ECameraFocusMethod::Manual );
diff --git a/Engine/Plugins/Experimental/VirtualCamera/Source/VirtualCamera/Private/ConcertVirtualCamera.cpp b/Engine/Plugins/Experimental/VirtualCamera/Source/VirtualCamera/Private/ConcertVirtualCamera.cpp
--- a/Engine/Plugins/Experimental/VirtualCamera/Source/VirtualCamera/Private/ConcertVirtualCamera.cpp
+++ b/Engine/Plugins/Experimental/VirtualCamera/Source/VirtualCamera/Private/ConcertVirtualCamera.cpp
@@ -74,7 +74,7 @@ FConcertVirtualCameraCameraData::FConcertVirtualCameraCameraData(const AActor* I
 	}
 	else
 	{
-		FilmbackSettings = InCineCameraComponent->FilmbackSettings;
+		FilmbackSettings = InCineCameraComponent->Filmback;
 	}
 }
 
@@ -90,7 +90,7 @@ void FConcertVirtualCameraCameraData::ApplyTo(AActor* InOwner, UCineCameraCompon
 	InCineCameraComponent->CurrentFocalLength = CurrentFocalLength;
 	InCineCameraComponent->FocusSettings = FocusSettings.ToCameraFocusSettings();
 	InCineCameraComponent->LensSettings = LensSettings;
-	InCineCameraComponent->FilmbackSettings = FilmbackSettings;
+	InCineCameraComponent->Filmback = FilmbackSettings;
 
 	if (UVirtualCameraCineCameraComponent* VCCineCamCmp = Cast<UVirtualCameraCineCameraComponent>(InCineCameraComponent))
 	{
diff --git a/Engine/Plugins/Experimental/VirtualCamera/Source/VirtualCamera/Private/VirtualCameraCineCameraComponent.cpp b/Engine/Plugins/Experimental/VirtualCamera/Source/VirtualCamera/Private/VirtualCameraCineCameraComponent.cpp
--- a/Engine/Plugins/Experimental/VirtualCamera/Source/VirtualCamera/Private/VirtualCameraCineCameraComponent.cpp
+++ b/Engine/Plugins/Experimental/VirtualCamera/Source/VirtualCamera/Private/VirtualCameraCineCameraComponent.cpp
@@ -152,13 +152,13 @@ void UVirtualCameraCineCameraComponent::UpdateCameraView()
 		// Set the actual camera filmback settings
 		if (ViewSizeAdjustmentForMatte >= 1.f)
 		{
-			FilmbackSettings.SensorWidth = DesiredFilmbackSettings.SensorWidth / ViewSizeRatio.X;
-			FilmbackSettings.SensorHeight = DesiredFilmbackSettings.SensorHeight / ViewSizeRatio.Y;
+			Filmback.SensorWidth = DesiredFilmbackSettings.SensorWidth / ViewSizeRatio.X;
+			Filmback.SensorHeight = DesiredFilmbackSettings.SensorHeight / ViewSizeRatio.Y;
 		}
 		else
 		{
-			FilmbackSettings.SensorWidth = DesiredFilmbackSettings.SensorWidth / ViewSizeRatio.X * ViewSizeAdjustmentForMatte;
-			FilmbackSettings.SensorHeight = DesiredFilmbackSettings.SensorHeight / ViewSizeRatio.Y * ViewSizeAdjustmentForMatte;
+			Filmback.SensorWidth = DesiredFilmbackSettings.SensorWidth / ViewSizeRatio.X * ViewSizeAdjustmentForMatte;
+			Filmback.SensorHeight = DesiredFilmbackSettings.SensorHeight / ViewSizeRatio.Y * ViewSizeAdjustmentForMatte;
 		}
 	}
 }
diff --git a/Engine/Plugins/Experimental/VirtualCamera/Source/VirtualCamera/Private/VirtualCameraPawnBase.cpp b/Engine/Plugins/Experimental/VirtualCamera/Source/VirtualCamera/Private/VirtualCameraPawnBase.cpp
--- a/Engine/Plugins/Experimental/VirtualCamera/Source/VirtualCamera/Private/VirtualCameraPawnBase.cpp
+++ b/Engine/Plugins/Experimental/VirtualCamera/Source/VirtualCamera/Private/VirtualCameraPawnBase.cpp
@@ -146,8 +146,8 @@ FString AVirtualCameraPawnBase::SavePreset(const bool bSaveCameraSettings, const
 	{
 		PresetToAdd.CameraSettings.FocalLength = CineCamera->GetCurrentFocalLength();
 		PresetToAdd.CameraSettings.Aperture = CineCamera->GetCurrentAperture();
-		PresetToAdd.CameraSettings.FilmbackWidth = CineCamera->FilmbackSettings.SensorWidth;
-		PresetToAdd.CameraSettings.FilmbackHeight = CineCamera->FilmbackSettings.SensorHeight;	
+		PresetToAdd.CameraSettings.FilmbackWidth = CineCamera->Filmback.SensorWidth;
+		PresetToAdd.CameraSettings.FilmbackHeight = CineCamera->Filmback.SensorHeight;	
 	}
 
 	if (MovementComponent)
@@ -234,8 +234,8 @@ void AVirtualCameraPawnBase::UpdateSettingsFromPreset(FVirtualCameraSettingsPres
 		{
 			CineCamera->CurrentAperture = PresetToLoad->CameraSettings.Aperture;
 			CineCamera->CurrentFocalLength = PresetToLoad->CameraSettings.FocalLength;
-			CineCamera->FilmbackSettings.SensorWidth = PresetToLoad->CameraSettings.FilmbackWidth;
-			CineCamera->FilmbackSettings.SensorHeight = PresetToLoad->CameraSettings.FilmbackHeight;
+			CineCamera->Filmback.SensorWidth = PresetToLoad->CameraSettings.FilmbackWidth;
+			CineCamera->Filmback.SensorHeight = PresetToLoad->CameraSettings.FilmbackHeight;
 		}
 	}
 	
@@ -300,7 +300,7 @@ FString AVirtualCameraPawnBase::TakeScreenshot()
 		// Apply aspect ratio restraints to image
 		if (CineCamera && World)
 		{
-			CineCamera->FilmbackSettings = CineCamera->DesiredFilmbackSettings;
+			CineCamera->Filmback = CineCamera->DesiredFilmbackSettings;
 			CineCamera->StopCameraViewUpdates();
 
 			FScreenshotRequest::RequestScreenshot(false);
@@ -425,8 +425,8 @@ void AVirtualCameraPawnBase::SaveSettings()
 
 	// Save filmback settings
 	SaveGameInstance->CameraSettings.FilmbackName = CineCamera->GetCurrentFilmbackName();
-	SaveGameInstance->CameraSettings.FilmbackWidth = CineCamera->FilmbackSettings.SensorWidth;
-	SaveGameInstance->CameraSettings.FilmbackHeight = CineCamera->FilmbackSettings.SensorHeight;
+	SaveGameInstance->CameraSettings.FilmbackWidth = CineCamera->Filmback.SensorWidth;
+	SaveGameInstance->CameraSettings.FilmbackHeight = CineCamera->Filmback.SensorHeight;
 	SaveGameInstance->CameraSettings.MatteOpacity = CineCamera->MatteOpacity;
 
 	// Save axis settings
@@ -504,8 +504,8 @@ void AVirtualCameraPawnBase::LoadSettings()
 	if (!CineCamera->SetFilmbackPresetOption(SaveGameInstance->CameraSettings.FilmbackName))
 	{
 		// If name isn't found, use backup settings
-		CineCamera->FilmbackSettings.SensorWidth = SaveGameInstance->CameraSettings.FilmbackWidth;
-		CineCamera->FilmbackSettings.SensorHeight = SaveGameInstance->CameraSettings.FilmbackHeight;
+		CineCamera->Filmback.SensorWidth = SaveGameInstance->CameraSettings.FilmbackWidth;
+		CineCamera->Filmback.SensorHeight = SaveGameInstance->CameraSettings.FilmbackHeight;
 	}
 	CineCamera->MatteOpacity = SaveGameInstance->CameraSettings.MatteOpacity;
 
diff --git a/Engine/Plugins/Experimental/VirtualCamera/Source/VirtualCamera/Private/VirtualCameraPlayerControllerBase.cpp b/Engine/Plugins/Experimental/VirtualCamera/Source/VirtualCamera/Private/VirtualCameraPlayerControllerBase.cpp
--- a/Engine/Plugins/Experimental/VirtualCamera/Source/VirtualCamera/Private/VirtualCameraPlayerControllerBase.cpp
+++ b/Engine/Plugins/Experimental/VirtualCamera/Source/VirtualCamera/Private/VirtualCameraPlayerControllerBase.cpp
@@ -499,7 +499,7 @@ void AVirtualCameraPlayerControllerBase::PilotTargetedCamera(AVirtualCameraPawnB
 		TargetCameraComponent->CurrentFocalLength = CameraToFollow->CurrentFocalLength;
 		TargetCameraComponent->FocusSettings = CameraToFollow->FocusSettings;
 		TargetCameraComponent->LensSettings = CameraToFollow->LensSettings;
-		TargetCameraComponent->FilmbackSettings = CameraToFollow->FilmbackSettings;
+		TargetCameraComponent->Filmback = CameraToFollow->Filmback;
 	}
 	TargetCameraComponent->SetRelativeLocationAndRotation(CameraToFollow->GetRelativeLocation(), CameraToFollow->GetRelativeRotation());
 }
diff --git a/Engine/Plugins/Importers/USDImporter/Source/USDUtilities/Private/USDPrimConversion.cpp b/Engine/Plugins/Importers/USDImporter/Source/USDUtilities/Private/USDPrimConversion.cpp
--- a/Engine/Plugins/Importers/USDImporter/Source/USDUtilities/Private/USDPrimConversion.cpp
+++ b/Engine/Plugins/Importers/USDImporter/Source/USDUtilities/Private/USDPrimConversion.cpp
@@ -63,8 +63,8 @@ bool UsdToUnreal::ConvertGeomCamera( const pxr::UsdStageRefPtr& Stage, const pxr
 
 	CameraComponent.CurrentAperture = UsdUtils::GetUsdValue< float >( GeomCamera.GetFStopAttr(), EvalTime );
 
-	CameraComponent.FilmbackSettings.SensorWidth = UsdUtils::GetUsdValue< float >( GeomCamera.GetHorizontalApertureAttr(), EvalTime );
-	CameraComponent.FilmbackSettings.SensorHeight = UsdUtils::GetUsdValue< float >( GeomCamera.GetVerticalApertureAttr(), EvalTime );
+	CameraComponent.Filmback.SensorWidth = UsdUtils::GetUsdValue< float >( GeomCamera.GetHorizontalApertureAttr(), EvalTime );
+	CameraComponent.Filmback.SensorHeight = UsdUtils::GetUsdValue< float >( GeomCamera.GetVerticalApertureAttr(), EvalTime );
 
 	return true;
 }
diff --git a/Engine/Source/Editor/MovieSceneTools/Private/MovieSceneToolHelpers.cpp b/Engine/Source/Editor/MovieSceneTools/Private/MovieSceneToolHelpers.cpp
--- a/Engine/Source/Editor/MovieSceneTools/Private/MovieSceneToolHelpers.cpp
+++ b/Engine/Source/Editor/MovieSceneTools/Private/MovieSceneToolHelpers.cpp
@@ -1296,8 +1296,8 @@ void MovieSceneToolHelpers::CopyCameraProperties(FbxCamera* CameraNode, AActor*
 		CameraComponent = CineCameraActor->GetCineCameraComponent();
 
 		UCineCameraComponent* CineCameraComponent = CineCameraActor->GetCineCameraComponent();
-		CineCameraComponent->FilmbackSettings.SensorWidth = FUnitConversion::Convert(ApertureWidth, EUnit::Inches, EUnit::Millimeters);
-		CineCameraComponent->FilmbackSettings.SensorHeight = FUnitConversion::Convert(ApertureHeight, EUnit::Inches, EUnit::Millimeters);
+		CineCameraComponent->Filmback.SensorWidth = FUnitConversion::Convert(ApertureWidth, EUnit::Inches, EUnit::Millimeters);
+		CineCameraComponent->Filmback.SensorHeight = FUnitConversion::Convert(ApertureHeight, EUnit::Inches, EUnit::Millimeters);
 		CineCameraComponent->FocusSettings.ManualFocusDistance = CameraNode->FocusDistance;
 		if (FocalLength < CineCameraComponent->LensSettings.MinFocalLength)
 		{
diff --git a/Engine/Source/Editor/UnrealEd/Private/Fbx/FbxMainExport.cpp b/Engine/Source/Editor/UnrealEd/Private/Fbx/FbxMainExport.cpp
--- a/Engine/Source/Editor/UnrealEd/Private/Fbx/FbxMainExport.cpp
+++ b/Engine/Source/Editor/UnrealEd/Private/Fbx/FbxMainExport.cpp
@@ -619,8 +619,8 @@ void FFbxExporter::FillFbxCameraAttribute(FbxNode* ParentNode, FbxCamera* Camera
 		UCineCameraComponent* CineCameraComponent = Cast<UCineCameraComponent>(CameraComponent);
 		if (CineCameraComponent)
 		{
-			ApertureWidthInInches = FUnitConversion::Convert(CineCameraComponent->FilmbackSettings.SensorWidth, EUnit::Millimeters, EUnit::Inches);
-			ApertureHeightInInches = FUnitConversion::Convert(CineCameraComponent->FilmbackSettings.SensorHeight, EUnit::Millimeters, EUnit::Inches);
+			ApertureWidthInInches = FUnitConversion::Convert(CineCameraComponent->Filmback.SensorWidth, EUnit::Millimeters, EUnit::Inches);
+			ApertureHeightInInches = FUnitConversion::Convert(CineCameraComponent->Filmback.SensorHeight, EUnit::Millimeters, EUnit::Inches);
 			FocalLength = CineCameraComponent->CurrentFocalLength;
 			FocusDistance = CineCameraComponent->FocusSettings.ManualFocusDistance;
 		}
diff --git a/Engine/Source/Editor/UnrealEd/Private/Fbx/FbxSceneImportFactory.cpp b/Engine/Source/Editor/UnrealEd/Private/Fbx/FbxSceneImportFactory.cpp
--- a/Engine/Source/Editor/UnrealEd/Private/Fbx/FbxSceneImportFactory.cpp
+++ b/Engine/Source/Editor/UnrealEd/Private/Fbx/FbxSceneImportFactory.cpp
@@ -1289,8 +1289,8 @@ USceneComponent *CreateCameraComponent(AActor *ParentActor, TSharedPtr<FFbxCamer
 	CameraComponent->SetOrthoFarClipPlane(CameraInfo->FarPlane);
 	CameraComponent->SetOrthoWidth(CameraInfo->AspectWidth);
 	CameraComponent->SetFieldOfView(CameraInfo->FieldOfView);
-	CameraComponent->FilmbackSettings.SensorWidth = FUnitConversion::Convert(CameraInfo->ApertureWidth, EUnit::Inches, EUnit::Millimeters);
-	CameraComponent->FilmbackSettings.SensorHeight = FUnitConversion::Convert(CameraInfo->ApertureHeight, EUnit::Inches, EUnit::Millimeters);
+	CameraComponent->Filmback.SensorWidth = FUnitConversion::Convert(CameraInfo->ApertureWidth, EUnit::Inches, EUnit::Millimeters);
+	CameraComponent->Filmback.SensorHeight = FUnitConversion::Convert(CameraInfo->ApertureHeight, EUnit::Inches, EUnit::Millimeters);
 	CameraComponent->LensSettings.MaxFocalLength = CameraInfo->FocalLength;
 	CameraComponent->LensSettings.MinFocalLength = CameraInfo->FocalLength;
 	CameraComponent->FocusSettings.FocusMethod = ECameraFocusMethod::None;
diff --git a/Engine/Source/Runtime/CinematicCamera/Private/CineCameraComponent.cpp b/Engine/Source/Runtime/CinematicCamera/Private/CineCameraComponent.cpp
--- a/Engine/Source/Runtime/CinematicCamera/Private/CineCameraComponent.cpp
+++ b/Engine/Source/Runtime/CinematicCamera/Private/CineCameraComponent.cpp
@@ -2,6 +2,7 @@
 
 #include "CineCameraComponent.h"
 #include "UObject/CineCameraObjectVersion.h"
+#include "UObject/ReleaseObjectVersion.h"
 #include "UObject/ConstructorHelpers.h"
 #include "GameFramework/Actor.h"
 #include "Engine/World.h"
@@ -25,8 +26,8 @@ UCineCameraComponent::UCineCameraComponent()
 {
 	// Super 35mm 4 Perf
 	// These will be overridden if valid default presets are specified in ini
-	FilmbackSettings.SensorWidth = 24.89f;
-	FilmbackSettings.SensorHeight = 18.67;
+	Filmback.SensorWidth = 24.89f;
+	Filmback.SensorHeight = 18.67;
 	LensSettings.MinFocalLength = 50.f;
 	LensSettings.MaxFocalLength = 50.f;
 	LensSettings.MinFStop = 2.f;
@@ -50,7 +51,8 @@ UCineCameraComponent::UCineCameraComponent()
 	if (Template)
 	{
 		// default filmback
-		SetFilmbackPresetByNameInternal(Template->DefaultFilmbackPreset);
+		SetFilmbackPresetByNameInternal(Template->DefaultFilmbackPreset, Filmback);
+		SetFilmbackPresetByNameInternal(Template->DefaultFilmbackPresetName_DEPRECATED, FilmbackSettings_DEPRECATED);
 		SetLensPresetByNameInternal(Template->DefaultLensPresetName);
 		// other lens defaults
 		CurrentAperture = Template->DefaultLensFStop;
@@ -78,6 +80,7 @@ UCineCameraComponent::UCineCameraComponent()
 void UCineCameraComponent::Serialize(FArchive& Ar)
 {
 	Ar.UsingCustomVersion(FCineCameraObjectVersion::GUID);
+	Ar.UsingCustomVersion(FReleaseObjectVersion::GUID);
 
 	Super::Serialize(Ar);
 }
@@ -91,12 +94,39 @@ void UCineCameraComponent::PostInitProperties()
 
 void UCineCameraComponent::PostLoad()
 {
-	UCineCameraComponent* Template = Cast<UCineCameraComponent>(GetArchetype());
-	if (Template)
+	if (GetLinkerCustomVersion(FReleaseObjectVersion::GUID) < FReleaseObjectVersion::DeprecateFilmbackSettings)
 	{
-		if (GetLinkerCustomVersion(FCineCameraObjectVersion::GUID) < FCineCameraObjectVersion::ChangeDefaultFilmbackToDigitalFilm)
+		bool bUpgradeFilmback = true;
+		if (GetLinkerCustomVersion(FCineCameraObjectVersion::GUID) == FCineCameraObjectVersion::ChangeDefaultFilmbackToDigitalFilm)
+		{
+			UCineCameraComponent* Template = Cast<UCineCameraComponent>(GetArchetype());
+			if (Template)
+			{
+				TArray<FNamedFilmbackPreset> const& Presets = UCineCameraComponent::GetFilmbackPresets();
+				int32 const NumPresets = Presets.Num();
+				for (int32 PresetIdx = 0; PresetIdx < NumPresets; ++PresetIdx)
+				{
+					FNamedFilmbackPreset const& P = Presets[PresetIdx];
+
+					// ChangeDefaultFilmbackToDigitalFilm was pre 4.24, but post 4.23. In that case, the filmback settings would have been DSLR 
+					// and RecalcDerivedData would not have been called yet, which equates to SensorAspectRatio being left at 1.33f. This isn't
+					// ideal for detecting this case, but it's the best notion of whether upgrading this film back should be skipped and get its 
+					// values from the default template object, which is the new Digital Film default.
+					if (P.FilmbackSettings == FilmbackSettings_DEPRECATED && FilmbackSettings_DEPRECATED.SensorAspectRatio == 1.33f)
+					{
+						if (P.Name == Template->DefaultFilmbackPresetName_DEPRECATED)
+						{
+							bUpgradeFilmback = false;
+							break;
+						}
+					}
+				}
+			}
+		}
+
+		if (bUpgradeFilmback)
 		{
-			SetFilmbackPresetByName(Template->DefaultFilmbackPresetName_DEPRECATED);
+			Filmback = FilmbackSettings_DEPRECATED;
 		}
 	}
 
@@ -200,7 +230,7 @@ void UCineCameraComponent::SetFieldOfView(float InFieldOfView)
 {
 	Super::SetFieldOfView(InFieldOfView);
 
-	CurrentFocalLength = (FilmbackSettings.SensorWidth / 2.f) / FMath::Tan(FMath::DegreesToRadians(InFieldOfView / 2.f));
+	CurrentFocalLength = (Filmback.SensorWidth / 2.f) / FMath::Tan(FMath::DegreesToRadians(InFieldOfView / 2.f));
 }
 
 void UCineCameraComponent::SetCurrentFocalLength(const float& InFocalLength)
@@ -212,14 +242,14 @@ void UCineCameraComponent::SetCurrentFocalLength(const float& InFocalLength)
 float UCineCameraComponent::GetHorizontalFieldOfView() const
 {
 	return (CurrentFocalLength > 0.f)
-		? FMath::RadiansToDegrees(2.f * FMath::Atan(FilmbackSettings.SensorWidth / (2.f * CurrentFocalLength)))
+		? FMath::RadiansToDegrees(2.f * FMath::Atan(Filmback.SensorWidth / (2.f * CurrentFocalLength)))
 		: 0.f;
 }
 
 float UCineCameraComponent::GetVerticalFieldOfView() const
 {
 	return (CurrentFocalLength > 0.f)
-		? FMath::RadiansToDegrees(2.f * FMath::Atan(FilmbackSettings.SensorHeight / (2.f * CurrentFocalLength)))
+		? FMath::RadiansToDegrees(2.f * FMath::Atan(Filmback.SensorHeight / (2.f * CurrentFocalLength)))
 		: 0.f;
 }
 
@@ -230,7 +260,7 @@ FString UCineCameraComponent::GetFilmbackPresetName() const
 	for (int32 PresetIdx = 0; PresetIdx < NumPresets; ++PresetIdx)
 	{
 		FNamedFilmbackPreset const& P = Presets[PresetIdx];
-		if (P.FilmbackSettings == FilmbackSettings)
+		if (P.FilmbackSettings == Filmback)
 		{
 			return P.Name;
 		}
@@ -241,12 +271,12 @@ FString UCineCameraComponent::GetFilmbackPresetName() const
 
 void UCineCameraComponent::SetFilmbackPresetByName(const FString& InPresetName)
 {
-	SetFilmbackPresetByNameInternal(InPresetName);
+	SetFilmbackPresetByNameInternal(InPresetName, Filmback);
 	// Explicitely call RecalcDerivedData() when invoked via Blueprint, since no other method (incl. PostEditChangeProperty) will trigger
 	RecalcDerivedData();
 }
 
-void UCineCameraComponent::SetFilmbackPresetByNameInternal(const FString& InPresetName)
+void UCineCameraComponent::SetFilmbackPresetByNameInternal(const FString& InPresetName, FCameraFilmbackSettings& InOutFilmbackSettings)
 {
 	TArray<FNamedFilmbackPreset> const& Presets = UCineCameraComponent::GetFilmbackPresets();
 	int32 const NumPresets = Presets.Num();
@@ -255,7 +285,7 @@ void UCineCameraComponent::SetFilmbackPresetByNameInternal(const FString& InPres
 		FNamedFilmbackPreset const& P = Presets[PresetIdx];
 		if (P.Name == InPresetName)
 		{
-			FilmbackSettings = P.FilmbackSettings;
+			InOutFilmbackSettings = P.FilmbackSettings;
 			break;
 		}
 	}
@@ -334,8 +364,8 @@ void UCineCameraComponent::RecalcDerivedData()
 	FocusSettings.ManualFocusDistance = FMath::Max(FocusSettings.ManualFocusDistance, MinFocusDistInWorldUnits);
 
 	FieldOfView = GetHorizontalFieldOfView();
-	FilmbackSettings.SensorAspectRatio = (FilmbackSettings.SensorHeight > 0.f) ? (FilmbackSettings.SensorWidth / FilmbackSettings.SensorHeight) : 0.f;
-	AspectRatio = FilmbackSettings.SensorAspectRatio;
+	Filmback.SensorAspectRatio = (Filmback.SensorHeight > 0.f) ? (Filmback.SensorWidth / Filmback.SensorHeight) : 0.f;
+	AspectRatio = Filmback.SensorAspectRatio;
 
 #if WITH_EDITORONLY_DATA
 	CurrentHorizontalFOV = FieldOfView;			// informational variable only, for editor users
@@ -395,8 +425,8 @@ void UCineCameraComponent::GetCameraView(float DeltaTime, FMinimalViewInfo& Desi
 #if WITH_EDITOR
 FText UCineCameraComponent::GetFilmbackText() const
 {
-	const float SensorWidth = FilmbackSettings.SensorWidth;
-	const float SensorHeight = FilmbackSettings.SensorHeight;
+	const float SensorWidth = Filmback.SensorWidth;
+	const float SensorHeight = Filmback.SensorHeight;
 
 	// Search presets for one that matches
 	const FNamedFilmbackPreset* Preset = UCineCameraComponent::GetFilmbackPresets().FindByPredicate([&](const FNamedFilmbackPreset& InPreset) {
@@ -488,7 +518,7 @@ void UCineCameraComponent::UpdateCameraLens(float DeltaTime, FMinimalViewInfo& D
 		DesiredView.PostProcessSettings.DepthOfFieldFocalDistance = CurrentFocusDistance;
 
 		DesiredView.PostProcessSettings.bOverride_DepthOfFieldSensorWidth = true;
-		DesiredView.PostProcessSettings.DepthOfFieldSensorWidth = FilmbackSettings.SensorWidth;
+		DesiredView.PostProcessSettings.DepthOfFieldSensorWidth = Filmback.SensorWidth;
 	}
 }
 
diff --git a/Engine/Source/Runtime/CinematicCamera/Public/CineCameraComponent.h b/Engine/Source/Runtime/CinematicCamera/Public/CineCameraComponent.h
--- a/Engine/Source/Runtime/CinematicCamera/Public/CineCameraComponent.h
+++ b/Engine/Source/Runtime/CinematicCamera/Public/CineCameraComponent.h
@@ -225,9 +225,12 @@ public:
 #if WITH_EDITOR
 	virtual FText GetFilmbackText() const override;
 #endif
+	UPROPERTY()
+	FCameraFilmbackSettings FilmbackSettings_DEPRECATED;
+
 	/** Controls the filmback of the camera. */
 	UPROPERTY(Interp, EditAnywhere, BlueprintReadWrite, Category = "Current Camera Settings")
-	FCameraFilmbackSettings FilmbackSettings;
+	FCameraFilmbackSettings Filmback;
 
 	/** Controls the camera's lens. */
 	UPROPERTY(EditAnywhere, BlueprintReadWrite, Category = "Current Camera Settings")
@@ -386,7 +389,7 @@ private:
 	float GetDesiredFocusDistance(const FVector& InLocation) const;
 	float GetWorldToMetersScale() const;
 	void SetLensPresetByNameInternal(const FString& InPresetName);
-	void SetFilmbackPresetByNameInternal(const FString& InPresetName);
+	void SetFilmbackPresetByNameInternal(const FString& InPresetName, FCameraFilmbackSettings& InOutFilmbackSettings);
 
 #if WITH_EDITORONLY_DATA
 	void CreateDebugFocusPlane();
diff --git a/Engine/Source/Runtime/Core/Public/UObject/ReleaseObjectVersion.h b/Engine/Source/Runtime/Core/Public/UObject/ReleaseObjectVersion.h
--- a/Engine/Source/Runtime/Core/Public/UObject/ReleaseObjectVersion.h
+++ b/Engine/Source/Runtime/Core/Public/UObject/ReleaseObjectVersion.h
@@ -93,6 +93,9 @@ struct CORE_API FReleaseObjectVersion
 		// Groom asset version
 		GroomAssetVersion3,
 
+		// Upgrade filmback
+		DeprecateFilmbackSettings,
+
 		// -----<new versions can be added above this line>-------------------------------------------------
 		VersionPlusOne,
 		LatestVersion = VersionPlusOne - 1
