Cine Camera: Change and deprecate film back

diff --git a/Engine/Config/BaseEngine.ini b/Engine/Config/BaseEngine.ini
--- a/Engine/Config/BaseEngine.ini
+++ b/Engine/Config/BaseEngine.ini
@@ -2226,7 +2226,8 @@ ServerUpdateCameraTimeout=2.0
 +FilmbackPresets=(Name="APS-C (Canon)",FilmbackSettings=(SensorWidth=22.2,SensorHeight=14.8))
 +FilmbackPresets=(Name="Full Frame DSLR",FilmbackSettings=(SensorWidth=36,SensorHeight=24))
 +FilmbackPresets=(Name="Micro Four Thirds",FilmbackSettings=(SensorWidth=17.3,SensorHeight=13))
-DefaultFilmbackPresetName="16:9 DSLR"
+DefaultFilmbackPresetName="16:9 DSLR" ; Deprecated, here for backwards compatibility
+DefaultFilmbackPreset="16:9 Digital Film"
 +LensPresets=(Name="12mm Prime f/2.8",LensSettings=(MinFocalLength=12,MaxFocalLength=12,MinFStop=2.8,MaxFStop=22,DiaphragmBladeCount=7))
 +LensPresets=(Name="30mm Prime f/1.4",LensSettings=(MinFocalLength=30,MaxFocalLength=30,MinFStop=1.4,MaxFStop=22,DiaphragmBladeCount=7))
 +LensPresets=(Name="50mm Prime f/1.8",LensSettings=(MinFocalLength=50,MaxFocalLength=50,MinFStop=1.8,MaxFStop=22,DiaphragmBladeCount=7))
diff --git a/Engine/Source/Runtime/CinematicCamera/Private/CineCameraComponent.cpp b/Engine/Source/Runtime/CinematicCamera/Private/CineCameraComponent.cpp
--- a/Engine/Source/Runtime/CinematicCamera/Private/CineCameraComponent.cpp
+++ b/Engine/Source/Runtime/CinematicCamera/Private/CineCameraComponent.cpp
@@ -1,6 +1,7 @@
 // Copyright 1998-2019 Epic Games, Inc. All Rights Reserved.
 
 #include "CineCameraComponent.h"
+#include "UObject/CineCameraObjectVersion.h"
 #include "UObject/ConstructorHelpers.h"
 #include "GameFramework/Actor.h"
 #include "Engine/World.h"
@@ -49,7 +50,7 @@ UCineCameraComponent::UCineCameraComponent()
 	if (Template)
 	{
 		// default filmback
-		SetFilmbackPresetByName(Template->DefaultFilmbackPresetName);
+		SetFilmbackPresetByName(Template->DefaultFilmbackPreset);
 		SetLensPresetByName(Template->DefaultLensPresetName);
 		// other lens defaults
 		CurrentAperture = Template->DefaultLensFStop;
@@ -74,6 +75,13 @@ UCineCameraComponent::UCineCameraComponent()
 #endif
 }
 
+void UCineCameraComponent::Serialize(FArchive& Ar)
+{
+	Ar.UsingCustomVersion(FCineCameraObjectVersion::GUID);
+
+	Super::Serialize(Ar);
+}
+
 void UCineCameraComponent::PostInitProperties()
 {
 	Super::PostInitProperties();
@@ -83,6 +91,15 @@ void UCineCameraComponent::PostInitProperties()
 
 void UCineCameraComponent::PostLoad()
 {
+	UCineCameraComponent* Template = Cast<UCineCameraComponent>(GetArchetype());
+	if (Template)
+	{
+		if (GetLinkerCustomVersion(FCineCameraObjectVersion::GUID) < FCineCameraObjectVersion::ChangeDefaultFilmbackToDigitalFilm)
+		{
+			SetFilmbackPresetByName(Template->DefaultFilmbackPresetName_DEPRECATED);
+		}
+	}
+
 	RecalcDerivedData();
 	bResetInterpolation = true;
 	Super::PostLoad();
diff --git a/Engine/Source/Runtime/CinematicCamera/Public/CineCameraComponent.h b/Engine/Source/Runtime/CinematicCamera/Public/CineCameraComponent.h
--- a/Engine/Source/Runtime/CinematicCamera/Public/CineCameraComponent.h
+++ b/Engine/Source/Runtime/CinematicCamera/Public/CineCameraComponent.h
@@ -307,6 +307,7 @@ protected:
 	
 	/// @endcond
 	
+	virtual void Serialize(FArchive& Ar) override;
 	virtual void PostInitProperties() override;
 	virtual void OnRegister() override;
 	virtual void TickComponent(float DeltaTime, ELevelTick TickType, FActorComponentTickFunction* ThisTickFunction) override;
@@ -344,9 +345,13 @@ protected:
 	UPROPERTY(config)
 	TArray<FNamedLensPreset> LensPresets;
 
+	/** Deprecated. See DefaultFilmbackPreset */
+	UPROPERTY(config)
+	FString DefaultFilmbackPresetName_DEPRECATED;
+
 	/** Name of the default filmback preset */
 	UPROPERTY(config)
-	FString DefaultFilmbackPresetName;
+	FString DefaultFilmbackPreset;
 
 	/** Name of the default lens preset */
 	UPROPERTY(config)
diff --git a/Engine/Source/Runtime/Core/Private/UObject/DevObjectVersion.cpp b/Engine/Source/Runtime/Core/Private/UObject/DevObjectVersion.cpp
--- a/Engine/Source/Runtime/Core/Private/UObject/DevObjectVersion.cpp
+++ b/Engine/Source/Runtime/Core/Private/UObject/DevObjectVersion.cpp
@@ -24,6 +24,7 @@
 #include "UObject/AutomationObjectVersion.h"
 #include "UObject/NiagaraObjectVersion.h"
 #include "UObject/DestructionObjectVersion.h"
+#include "UObject/CineCameraObjectVersion.h"
 
 DEFINE_LOG_CATEGORY_STATIC(LogDevObjectVersion, Log, All);
 
@@ -163,3 +164,9 @@ FDevVersionRegistration GRegisterNiagaraObjectVersion(FNiagaraObjectVersion::GUI
 const FGuid FDestructionObjectVersion::GUID(0x174F1F0B, 0xB4C645A5, 0xB13F2EE8, 0xD0FB917D);
 // Register Destruction custom version with Core
 FDevVersionRegistration GRegisterDestructionObjectVersion(FDestructionObjectVersion::GUID, FDestructionObjectVersion::LatestVersion, TEXT("Dev-Destruction"));
+
+// Unique CineCamera Object version id
+const FGuid FCineCameraObjectVersion::GUID(0xB2E18506, 0x4273CFC2, 0xA54EF4BB, 0x758BBA07);
+// Register CineCamera custom version with Core
+FDevVersionRegistration GRegisterCineCameraObjectVersion(FCineCameraObjectVersion::GUID, FCineCameraObjectVersion::LatestVersion, TEXT("Dev-CineCamera"));
+
diff --git a/Engine/Source/Runtime/Core/Public/UObject/CineCameraObjectVersion.h b/Engine/Source/Runtime/Core/Public/UObject/CineCameraObjectVersion.h
new file mode 100644
--- /dev/null
+++ b/Engine/Source/Runtime/Core/Public/UObject/CineCameraObjectVersion.h
@@ -0,0 +1,28 @@
+// Copyright 1998-2019 Epic Games, Inc. All Rights Reserved.
+#pragma once
+
+#include "CoreTypes.h"
+#include "Misc/Guid.h"
+
+// Custom serialization version for cine camera
+struct CORE_API FCineCameraObjectVersion
+{
+	enum Type
+	{
+		// Before any version changes were made
+		BeforeCustomVersionWasAdded,
+
+		// Changed default filmback to digital film
+		ChangeDefaultFilmbackToDigitalFilm,
+
+		// -----<new versions can be added above this line>-------------------------------------------------
+		VersionPlusOne,
+		LatestVersion = VersionPlusOne - 1
+	};
+
+	// The GUID for this custom version number
+	const static FGuid GUID;
+
+private:
+	FCineCameraObjectVersion() {}
+};
