CameraRig: Check for valid world when updating preview meshes

diff --git a/Engine/Source/Runtime/CinematicCamera/Private/CameraRig_Rail.cpp b/Engine/Source/Runtime/CinematicCamera/Private/CameraRig_Rail.cpp
--- a/Engine/Source/Runtime/CinematicCamera/Private/CameraRig_Rail.cpp
+++ b/Engine/Source/Runtime/CinematicCamera/Private/CameraRig_Rail.cpp
@@ -150,6 +150,11 @@ void ACameraRig_Rail::UpdatePreviewMeshes()
 
 void ACameraRig_Rail::UpdateRailComponents()
 {
+	if (!GetWorld())
+	{
+		return;
+	}
+
 	if (RailSplineComponent && RailCameraMount)
 	{
 		float const SplineLen = RailSplineComponent->GetSplineLength();
