Camera: Change debug focus plane material to use unlit transluscent

#rnx

diff --git a/Engine/Source/Runtime/CinematicCamera/Private/CineCameraComponent.cpp b/Engine/Source/Runtime/CinematicCamera/Private/CineCameraComponent.cpp
--- a/Engine/Source/Runtime/CinematicCamera/Private/CineCameraComponent.cpp
+++ b/Engine/Source/Runtime/CinematicCamera/Private/CineCameraComponent.cpp
@@ -72,7 +72,7 @@ UCineCameraComponent::UCineCameraComponent()
 	static ConstructorHelpers::FObjectFinder<UStaticMesh> PlaneMesh(TEXT("/Engine/ArtTools/RenderToTexture/Meshes/S_1_Unit_Plane.S_1_Unit_Plane"));
 	FocusPlaneVisualizationMesh = PlaneMesh.Object;
 
-	static ConstructorHelpers::FObjectFinder<UMaterial> PlaneMat(TEXT("/Engine/EngineDebugMaterials/M_SimpleTranslucent.M_SimpleTranslucent"));
+	static ConstructorHelpers::FObjectFinder<UMaterial> PlaneMat(TEXT("/Engine/EngineDebugMaterials/M_SimpleUnlitTranslucent.M_SimpleUnlitTranslucent"));
 	FocusPlaneVisualizationMaterial = PlaneMat.Object;
 #endif
 }
