Cine Camera: Fix custom camera filmback text getting incorrect values for focal length and aperture

diff --git a/Engine/Source/Runtime/CinematicCamera/Private/CineCameraComponent.cpp b/Engine/Source/Runtime/CinematicCamera/Private/CineCameraComponent.cpp
--- a/Engine/Source/Runtime/CinematicCamera/Private/CineCameraComponent.cpp
+++ b/Engine/Source/Runtime/CinematicCamera/Private/CineCameraComponent.cpp
@@ -446,7 +446,7 @@ FText UCineCameraComponent::GetFilmbackText() const
 	{
 		FNumberFormattingOptions Opts = FNumberFormattingOptions().SetMaximumFractionalDigits(1);
 		return FText::Format(
-			LOCTEXT("CustomFilmbackFormat", "Custom ({0}mm x {1}mm) | Zoom: {1}mm | Av: {2}"),
+			LOCTEXT("CustomFilmbackFormat", "Custom ({0}mm x {1}mm) | Zoom: {2}mm | Av: {3}"),
 			FText::AsNumber(SensorWidth, &Opts),
 			FText::AsNumber(SensorHeight, &Opts),
 			FText::AsNumber(CurrentFocalLength),
