Skip updating CameraRig_Rail preview meshes in game worlds

diff --git a/Engine/Source/Runtime/CinematicCamera/Private/CameraRig_Rail.cpp b/Engine/Source/Runtime/CinematicCamera/Private/CameraRig_Rail.cpp
--- a/Engine/Source/Runtime/CinematicCamera/Private/CameraRig_Rail.cpp
+++ b/Engine/Source/Runtime/CinematicCamera/Private/CameraRig_Rail.cpp
@@ -184,8 +184,12 @@ void ACameraRig_Rail::UpdateRailComponents()
 #if WITH_EDITOR
 	if (GIsEditor)
 	{
-		// set up preview mesh to match #todo
-		UpdatePreviewMeshes();
+		const UWorld* const MyWorld = GetWorld();
+		if (MyWorld && !MyWorld->IsGameWorld())
+		{
+			// set up preview mesh to match #todo
+			UpdatePreviewMeshes();
+		}
 	}
 #endif
 }
