Cameras: Set the PP blend weight to the current camera's blend weight.

diff --git a/Engine/Source/Runtime/CinematicCamera/Private/CineCameraComponent.cpp b/Engine/Source/Runtime/CinematicCamera/Private/CineCameraComponent.cpp
--- a/Engine/Source/Runtime/CinematicCamera/Private/CineCameraComponent.cpp
+++ b/Engine/Source/Runtime/CinematicCamera/Private/CineCameraComponent.cpp
@@ -518,6 +518,8 @@ void UCineCameraComponent::UpdateCameraLens(float DeltaTime, FMinimalViewInfo& D
 	else
 	{
 		// Update focus/DoF
+		DesiredView.PostProcessBlendWeight = PostProcessBlendWeight;
+
 		DesiredView.PostProcessSettings.bOverride_DepthOfFieldFstop = true;
 		DesiredView.PostProcessSettings.DepthOfFieldFstop = CurrentAperture;
 
