PR #9049: [UE5] Expose LookAtTrackingInterpSpeed to cinematics (Contributed by [REDACTED_PERSON])

diff --git a/Engine/Source/Runtime/CinematicCamera/Public/CineCameraActor.h b/Engine/Source/Runtime/CinematicCamera/Public/CineCameraActor.h
--- a/Engine/Source/Runtime/CinematicCamera/Public/CineCameraActor.h
+++ b/Engine/Source/Runtime/CinematicCamera/Public/CineCameraActor.h
@@ -35,7 +35,7 @@ struct FCameraLookatTrackingSettings
 	uint8 bDrawDebugLookAtTrackingPosition : 1;
 
 	/** Controls degree of smoothing. 0.f for no smoothing, higher numbers for faster/tighter tracking. */
-	UPROPERTY(EditAnywhere, BlueprintReadWrite, Category = "LookAt")
+	UPROPERTY(Interp, EditAnywhere, BlueprintReadWrite, Category = "LookAt")
 	float LookAtTrackingInterpSpeed;
 
 	/** Last known lookat tracking rotation (used during interpolation) */
