Cine Camera: Debug focus plane is now editor only again. There was a bug where editor only pins in a struct were visiible in blueprint and that caused a crash. That has been fixed, so debug focus plane properties can no longer be accessed in blueprint.

diff --git a/Engine/Source/Runtime/CinematicCamera/Public/CineCameraComponent.h b/Engine/Source/Runtime/CinematicCamera/Public/CineCameraComponent.h
--- a/Engine/Source/Runtime/CinematicCamera/Public/CineCameraComponent.h
+++ b/Engine/Source/Runtime/CinematicCamera/Public/CineCameraComponent.h
@@ -193,8 +193,7 @@ struct FCameraFocusSettings
 	UPROPERTY(EditAnywhere, BlueprintReadWrite, Category = "Tracking Focus Settings")
 	FCameraTrackingFocusSettings TrackingFocusSettings;
 
-//~ TODO: Make this editor only again once UE-43122 has been completed.
-//~	#if WITH_EDITORONLY_DATA
+#if WITH_EDITORONLY_DATA
 	/** True to draw a translucent plane at the current focus depth, for easy tweaking. */
 	UPROPERTY(Transient, EditAnywhere, Category = "Focus Settings")
 	uint8 bDrawDebugFocusPlane : 1;
@@ -202,7 +201,7 @@ struct FCameraFocusSettings
 	/** For customizing the focus plane color, in case the default doesn't show up well in your scene. */
 	UPROPERTY(EditAnywhere, Category = "Focus Settings", meta = (EditCondition = "bDrawDebugFocusPlane"))
 	FColor DebugFocusPlaneColor;
-//~	#endif 
+#endif 
 
 	/** True to use interpolation to smooth out changes in focus distance, false for focus distance changes to be instantaneous. */
 	UPROPERTY(EditAnywhere, BlueprintReadWrite, Category = "Focus Settings")
