Camera Rig: Rail mount visualization shouldn't rotate if bLockOrientationToRail is disabled. The attachment already handles this correctly.This only fixes the visualization.

diff --git a/Engine/Source/Runtime/CinematicCamera/Private/CameraRig_Rail.cpp b/Engine/Source/Runtime/CinematicCamera/Private/CameraRig_Rail.cpp
--- a/Engine/Source/Runtime/CinematicCamera/Private/CameraRig_Rail.cpp
+++ b/Engine/Source/Runtime/CinematicCamera/Private/CameraRig_Rail.cpp
@@ -148,9 +148,13 @@ void ACameraRig_Rail::UpdatePreviewMeshes()
 		// make visualization of the mount follow the contour of the rail
 		if (PreviewMesh_Mount)
 		{
-			float const SplineLen = RailSplineComponent->GetSplineLength();
-			FQuat const RailRot = RailSplineComponent->GetQuaternionAtDistanceAlongSpline(CurrentPositionOnRail*SplineLen, ESplineCoordinateSpace::World);
-			PreviewMesh_Mount->SetWorldRotation(RailRot);
+			if (bLockOrientationToRail)
+			{
+				float const SplineLen = RailSplineComponent->GetSplineLength();
+				FQuat const RailRot = RailSplineComponent->GetQuaternionAtDistanceAlongSpline(CurrentPositionOnRail*SplineLen, ESplineCoordinateSpace::World);
+				PreviewMesh_Mount->SetWorldRotation(RailRot);
+			}
+			
 			PreviewMesh_Mount->SetVisibility(bShowRailVisualization);
 			PreviewMesh_Mount->SetWorldScale3D(FVector(PreviewMeshScale, PreviewMeshScale, PreviewMeshScale));
 		}
