Implements anamorphic lens' squeeze factor and post desqueeze aspect ratio crop on the cinecamera

Notes:
* Model used for the squeeze factor is HorizontalFocalLength = LensFocalLength / SqueezeFactor; VerticalFocalLength = LensFocalLength;
* The renderer renders the desqueezed image directly because of not wanting to add support to maintain non squared pixel.
* Squeeze factor is only hooked in to DOF's CoC.
* Squeeze factor currently also affects depth blur which it shouldn't in theory, requires dual CoC shader permutations in DOF algorithm to do that correctly.
* Crop's model fits an aspect ratio directly into the desqueezed full sensor image.
* Crop settings are directly applied to the camera's FOV calculation of the camera to avoid rendering extra pixels.
* Custom crop aspect ratios can be added into the project's DefaultEngine.ini

diff --git a/Engine/Config/BaseEngine.ini b/Engine/Config/BaseEngine.ini
--- a/Engine/Config/BaseEngine.ini
+++ b/Engine/Config/BaseEngine.ini
@@ -2955,6 +2955,11 @@ DefaultFilmbackPreset="16:9 Digital Film"
 +LensPresets=(Name="70-200mm Zoom f/2.8",LensSettings=(MinFocalLength=70,MaxFocalLength=200,MinFStop=2.8,MaxFStop=22,DiaphragmBladeCount=7))
 +LensPresets=(Name="Universal Zoom",LensSettings=(MinFocalLength=4,MaxFocalLength=1000,MinFStop=1.2,MaxFStop=22,DiaphragmBladeCount=7))
 DefaultLensPresetName="Universal Zoom"
++CropPresets=(Name="No Crop",CropSettings=(AspectRatio=0))
++CropPresets=(Name="1.33 (4:3)",CropSettings=(AspectRatio=1.333333))
++CropPresets=(Name="1.77 (16:9)",CropSettings=(AspectRatio=1.777778))
++CropPresets=(Name="2.39",CropSettings=(AspectRatio=2.39))
+DefaultCropPresetName="No Crop"
 DefaultLensFocalLength=35
 DefaultLensFStop=2.8
 
diff --git a/Engine/Shaders/Private/DiaphragmDOF/DOFCommon.ush b/Engine/Shaders/Private/DiaphragmDOF/DOFCommon.ush
--- a/Engine/Shaders/Private/DiaphragmDOF/DOFCommon.ush
+++ b/Engine/Shaders/Private/DiaphragmDOF/DOFCommon.ush
@@ -145,6 +145,18 @@ static const int2 kOffsetsSquare3x3[9] = {
 };
 
 
+//------------------------------------------------------- PARAMETERS
+
+float CocInfinityRadius;
+float CocMinRadius;
+float CocMaxRadius;
+float CocSqueeze;
+float CocInvSqueeze;
+
+float DepthBlurRadius;
+float DepthBlurExponent;
+
+
 //------------------------------------------------------- DEBUG OUTPUT
 
 #if DEBUG_OUTPUT
@@ -154,29 +166,20 @@ RWTexture2D<float4> DebugOutput;
 
 //------------------------------------------------------- COMMON API.
 
-// is wrapping the DepthToCoc() function in this file for faster iteration
-// @return half res pixel radius
-float DepthToHalfResCocRadius(float4 CocModelParameters, float2 DepthBlurParameters, float SceneDepth)
+// Returns the coc radius
+float SceneDepthToCocRadius(float SceneDepth)
 {
 	const float Focus = View.DepthOfFieldFocalDistance; // TODO: rename this to DepthOfFieldFocusDistance.
-	const float InfinityCocRadius = CocModelParameters.x;
-	const float MinCocRadius = CocModelParameters.y;
-	const float MaxCocRadius = CocModelParameters.z;
 	
 	// Radius from the circle of confusion.
-	float CocRadius = ((SceneDepth - Focus) / SceneDepth) * InfinityCocRadius;
+	float CocRadius = ((SceneDepth - Focus) / SceneDepth) * CocInfinityRadius;
 
 	// Depth blur's radius.
-	float DepthBlurAbsRadius = (1.0 - exp2(-SceneDepth * DepthBlurParameters.x)) * DepthBlurParameters.y;
+	float DepthBlurAbsRadius = (1.0 - exp2(-SceneDepth * DepthBlurExponent)) * DepthBlurRadius;
 
-	float ReturnCoc = max(abs(CocRadius), DepthBlurAbsRadius);
-	if(CocRadius < 0.0) 
-	{
-		// Near CoC is using negative values.
-		ReturnCoc = -ReturnCoc;
-	}
+	float ReturnCoc = max(abs(CocRadius), DepthBlurAbsRadius) * sign(CocRadius);
 
-	return clamp(ReturnCoc, MinCocRadius, MaxCocRadius);
+	return clamp(ReturnCoc, CocMinRadius, CocMaxRadius);
 }
 
 // Return whether this is front element from CocRadius.
diff --git a/Engine/Shaders/Private/DiaphragmDOF/DOFGatherKernel.ush b/Engine/Shaders/Private/DiaphragmDOF/DOFGatherKernel.ush
--- a/Engine/Shaders/Private/DiaphragmDOF/DOFGatherKernel.ush
+++ b/Engine/Shaders/Private/DiaphragmDOF/DOFGatherKernel.ush
@@ -462,6 +462,7 @@ void GatherRingSamples(
 				};
 				
 				SamplePosOffset[0] = KernelScope.SampleCountToSamplePosition * BokehLUT.SampleLevel(GlobalBilinearWrappedSampler, LookUpUV[0], 0).xy;
+				SamplePosOffset[0].x *= CocInvSqueeze;
 
 				#if DIM_BOKEH_SIMULATION == BOKEH_SIMULATION_SIMMETRIC && 0
 					SamplePosOffset[1] = -SamplePosOffset[0];
@@ -484,19 +485,22 @@ void GatherRingSamples(
 				}
 
 				SamplePosOffset[1] = -SamplePosOffset[0];
-				SampleDistance = length(SamplePosOffset[0]);
+				SampleDistance = length(SamplePosOffset[0] * float2(CocSqueeze, 1.0));
 			}
 			#elif CONFIG_KERNEL_LAYOUT == KERNEL_LAYOUT_SCALABLE_DISK
 			{
 				// Always keep a disk shaped kernel.
 				SamplePosOffset[0] = CircleSamplePosOffset;
-				SamplePosOffset[1] = -CircleSamplePosOffset;
+				SamplePosOffset[0].x *= CocInvSqueeze;
+
+				SamplePosOffset[1] = -SamplePosOffset[0];
 				SampleDistance = RingRadius;
 			}
 			#else
 				#error Unknown gathering kernel layout.
 
 			#endif
+
 			
 			// Fetches the mirrored samples.
 			FGatherSample Sample[2];
diff --git a/Engine/Shaders/Private/DiaphragmDOF/DOFHybridScatterPixelShader.usf b/Engine/Shaders/Private/DiaphragmDOF/DOFHybridScatterPixelShader.usf
--- a/Engine/Shaders/Private/DiaphragmDOF/DOFHybridScatterPixelShader.usf
+++ b/Engine/Shaders/Private/DiaphragmDOF/DOFHybridScatterPixelShader.usf
@@ -83,6 +83,7 @@ void ScatterMainPS(
 	{
 		// Pixel coordinate of the center of the bokeh.
 		float2 PixelOffset = PixelOffset0 - ScatteringScaling * kSquare2x2[i];
+		PixelOffset.x *= rcp(CocInvSqueeze);
 
 		float Intersection;
 		#if DIM_BOKEH_SIMULATION
diff --git a/Engine/Shaders/Private/DiaphragmDOF/DOFHybridScatterVertexShader.usf b/Engine/Shaders/Private/DiaphragmDOF/DOFHybridScatterVertexShader.usf
--- a/Engine/Shaders/Private/DiaphragmDOF/DOFHybridScatterVertexShader.usf
+++ b/Engine/Shaders/Private/DiaphragmDOF/DOFHybridScatterVertexShader.usf
@@ -135,7 +135,7 @@ void ScatterMainVS(
 	// The center of the pixel group is 0.5 south east of BokehCenter0.
 	const float HalfPixelShift = 0.5;
 
-	float2 ViewportUV = ((MaxAbsCocRadius * CocRadiusToCircumscribedRadius + (HalfPixelShift + ANTI_ALIASING_FEATHER_OFFSET)) * VertexPos + BokehCenter0 + HalfPixelShift) * ViewportSize.zw; 
+	float2 ViewportUV = ((MaxAbsCocRadius * CocRadiusToCircumscribedRadius * float2(CocInvSqueeze, 1.0) + (HalfPixelShift + ANTI_ALIASING_FEATHER_OFFSET)) * VertexPos + BokehCenter0 + HalfPixelShift) * ViewportSize.zw; 
 
 	OutPosition = float4(ViewportUV.x * 2.0 - 1.0, 1.0 - ViewportUV.y * 2.0, 0.0, 1.0);
 }
diff --git a/Engine/Shaders/Private/DiaphragmDOF/DOFRecombine.usf b/Engine/Shaders/Private/DiaphragmDOF/DOFRecombine.usf
--- a/Engine/Shaders/Private/DiaphragmDOF/DOFRecombine.usf
+++ b/Engine/Shaders/Private/DiaphragmDOF/DOFRecombine.usf
@@ -135,8 +135,6 @@ float4 ViewportSize;
 uint4 ViewportRect;
 
 float2 TemporalJitterPixels;
-float4 CocModelParameters;
-float2 DepthBlurParameters;
 float2 DOFBufferUVMax;
 
 float4 SeparateTranslucencyBilinearUVMinMax;
@@ -216,11 +214,6 @@ float2 RotVec(float Radius, float Radians)
 	return Radius * float2(cos(Radians), sin(Radians));
 }
 
-float DepthToHalfResCocRadius(float WorldDepth)
-{
-	return DepthToHalfResCocRadius(CocModelParameters, DepthBlurParameters, WorldDepth);
-}
-
 void AmendAdditiveColorWithMaxTranslucency(inout float4 Color, inout float Translucency, float MaxTranslucency)
 {
 	if (Translucency < 1)
@@ -604,7 +597,7 @@ void FetchAndAccumulateSamplePair(
 			Sample[k].CocRadius = Sample[k].Color.a;
 		#else
 			Sample[k].Color = SceneColorInput.SampleLevel(GlobalPointClampedSampler, BufferUV, 0);
-			Sample[k].CocRadius = DepthToHalfResCocRadius(SampleWorldDepth(BufferUV));
+			Sample[k].CocRadius = SceneDepthToCocRadius(SampleWorldDepth(BufferUV));
 		#endif
 		
 		// Convert scene color alpha from translucency to opacity.
@@ -754,7 +747,7 @@ void RecombineMainCS(
 	#if CONFIG_FETCH_FULLRES_COC_FROM_ALPHA
 		CenterSample.CocRadius = CenterSample.Color.a;
 	#else
-		CenterSample.CocRadius = DepthToHalfResCocRadius(SampleWorldDepth(InputParameters.SceneBufferUV));
+		CenterSample.CocRadius = SceneDepthToCocRadius(SampleWorldDepth(InputParameters.SceneBufferUV));
 	#endif
 	
 	// Convert scene color alpha from translucency to opacity.
@@ -798,7 +791,7 @@ void RecombineMainCS(
 				#if CONFIG_FETCH_FULLRES_COC_FROM_ALPHA
 					float SampleCocRadius = SceneColorInput.SampleLevel(GlobalPointClampedSampler, SampleUV, 0).a;
 				#else
-					float SampleCocRadius = DepthToHalfResCocRadius(SampleWorldDepth(SampleUV));
+					float SampleCocRadius = SceneDepthToCocRadius(SampleWorldDepth(SampleUV));
 				#endif
 
 				float SampleAbsCocRadius = abs(SampleCocRadius);
diff --git a/Engine/Shaders/Private/DiaphragmDOF/DOFReduce.usf b/Engine/Shaders/Private/DiaphragmDOF/DOFReduce.usf
--- a/Engine/Shaders/Private/DiaphragmDOF/DOFReduce.usf
+++ b/Engine/Shaders/Private/DiaphragmDOF/DOFReduce.usf
@@ -478,7 +478,7 @@ void ReduceCS(
 		
 		// Intensity loss factor caused by the area of confusion.
 		// TODO: no need min()
-		float IntensityLossFactor = min(1, SafeRcp(PI * CocRadius * CocRadius, 1.0));
+		float IntensityLossFactor = min(1, SafeRcp(PI * CocRadius * CocRadius, 1.0) * CocSqueeze);
 
 		// Enqueue scattering group in foreground scatter draw list.
 		BRANCH
diff --git a/Engine/Shaders/Private/DiaphragmDOF/DOFSetup.usf b/Engine/Shaders/Private/DiaphragmDOF/DOFSetup.usf
--- a/Engine/Shaders/Private/DiaphragmDOF/DOFSetup.usf
+++ b/Engine/Shaders/Private/DiaphragmDOF/DOFSetup.usf
@@ -22,9 +22,6 @@
 
 uint4 ViewportRect;
 
-float4 CocModelParameters;
-float2 DepthBlurParameters;
-
 // x: full res basis; y: half res basis.
 float2 CocRadiusBasis;
 
@@ -83,7 +80,7 @@ void SetupCS(
 
 	float4 SceneColor = SceneColorTexture.SampleLevel(GlobalPointClampedSampler, SceneBufferUV, 0);
 	float DeviceZ = SceneDepthTexture.SampleLevel(GlobalPointClampedSampler, SceneBufferUV, 0).r;
-	float CocRadius = DepthToHalfResCocRadius(CocModelParameters, DepthBlurParameters, ConvertFromDeviceZ(DeviceZ));
+	float CocRadius = SceneDepthToCocRadius(ConvertFromDeviceZ(DeviceZ));
 	
 	// Convert from translucency to opacity.
 	SceneColor.a = 1 - SceneColor.a;
@@ -184,7 +181,7 @@ void SetupCS(uint2 DispatchThreadId : SV_DispatchThreadID)
 	float4 OutColor = SceneColorTexture.SampleLevel(GlobalPointClampedSampler, SceneBufferUV, 0);
 	
 	float DeviceZ = SceneDepthTexture.SampleLevel(GlobalPointClampedSampler, SceneBufferUV, 0).r;
-	float OutCocRadius = CocRadiusBasis.x * DepthToHalfResCocRadius(CocModelParameters, DepthBlurParameters, ConvertFromDeviceZ(DeviceZ));
+	float OutCocRadius = CocRadiusBasis.x * SceneDepthToCocRadius(ConvertFromDeviceZ(DeviceZ));
 
 	if (all(DispatchThreadId < ViewportRect.zw))
 	{
@@ -224,10 +221,10 @@ void SetupCS(
 		// Fetch Coc
 		float4 SceneDepthGather = SceneDepthTexture.Gather(GlobalPointClampedSampler, SceneBufferUV);
 		float CocRadii[4];
-		CocRadii[0] = DepthToHalfResCocRadius(CocModelParameters, DepthBlurParameters, ConvertFromDeviceZ(SceneDepthGather.w));
-		CocRadii[1] = DepthToHalfResCocRadius(CocModelParameters, DepthBlurParameters, ConvertFromDeviceZ(SceneDepthGather.z));
-		CocRadii[2] = DepthToHalfResCocRadius(CocModelParameters, DepthBlurParameters, ConvertFromDeviceZ(SceneDepthGather.x));
-		CocRadii[3] = DepthToHalfResCocRadius(CocModelParameters, DepthBlurParameters, ConvertFromDeviceZ(SceneDepthGather.y));
+		CocRadii[0] = SceneDepthToCocRadius(ConvertFromDeviceZ(SceneDepthGather.w));
+		CocRadii[1] = SceneDepthToCocRadius(ConvertFromDeviceZ(SceneDepthGather.z));
+		CocRadii[2] = SceneDepthToCocRadius(ConvertFromDeviceZ(SceneDepthGather.x));
+		CocRadii[3] = SceneDepthToCocRadius(ConvertFromDeviceZ(SceneDepthGather.y));
 
 		// Fetch color.
 		float4 SceneColors[4];
@@ -276,7 +273,7 @@ void SetupCS(
 			SceneColors[i].a = 1 - SceneColors[i].a;
 
 			float SampleDeviceZ = SceneDepthTexture.SampleLevel(GlobalPointClampedSampler, SampleUV, 0).r;
-			CocRadii[i] = DepthToHalfResCocRadius(CocModelParameters, DepthBlurParameters, ConvertFromDeviceZ(SampleDeviceZ));
+			CocRadii[i] = SceneDepthToCocRadius(ConvertFromDeviceZ(SampleDeviceZ));
 		}
 
 		// Downsample.
diff --git a/Engine/Source/Editor/DetailCustomizations/Private/CameraCropSettingsCustomization.cpp b/Engine/Source/Editor/DetailCustomizations/Private/CameraCropSettingsCustomization.cpp
new file mode 100644
--- /dev/null
+++ b/Engine/Source/Editor/DetailCustomizations/Private/CameraCropSettingsCustomization.cpp
@@ -0,0 +1,153 @@
+// Copyright Epic Games, Inc. All Rights Reserved.
+
+#include "CameraCropSettingsCustomization.h"
+#include "PropertyHandle.h"
+#include "IDetailChildrenBuilder.h"
+#include "DetailWidgetRow.h"
+#include "DetailLayoutBuilder.h"
+#include "ScopedTransaction.h"
+#include "CineCameraComponent.h"
+
+#define LOCTEXT_NAMESPACE "CameraCropSettingsCustomization"
+
+FCameraCropSettingsCustomization::FCameraCropSettingsCustomization()
+{
+	TArray<FNamedPlateCropPreset> const& Presets = UCineCameraComponent::GetCropPresets();
+
+	int32 const NumPresets = Presets.Num();
+	// first create preset combo list
+	PresetComboList.Empty(NumPresets + 1);
+
+	// put all presets in the list
+	for (FNamedPlateCropPreset const& P : Presets)
+	{
+		PresetComboList.Add(MakeShareable(new FString(P.Name)));
+	}
+
+	PresetComboList.Add(MakeShareable(new FString(TEXT("Custom..."))));
+}
+
+TSharedRef<IPropertyTypeCustomization> FCameraCropSettingsCustomization::MakeInstance()
+{
+	return MakeShareable(new FCameraCropSettingsCustomization);
+}
+
+void FCameraCropSettingsCustomization::CustomizeHeader(TSharedRef<IPropertyHandle> StructPropertyHandle, FDetailWidgetRow& HeaderRow, IPropertyTypeCustomizationUtils& CustomizationUtils)
+{
+	HeaderRow.
+		NameContent()
+		[
+			StructPropertyHandle->CreatePropertyNameWidget()
+		]
+		.ValueContent()
+		.MaxDesiredWidth(0.f)
+		[
+			SAssignNew(PresetComboBox, SComboBox< TSharedPtr<FString> >)
+			.OptionsSource(&PresetComboList)
+			.OnGenerateWidget(this, &FCameraCropSettingsCustomization::MakePresetComboWidget)
+			.OnSelectionChanged(this, &FCameraCropSettingsCustomization::OnPresetChanged)
+			.IsEnabled(FSlateApplication::Get().GetNormalExecutionAttribute())
+			.ContentPadding(2)
+			.Content()
+			[
+				SNew(STextBlock)
+				.Text(this, &FCameraCropSettingsCustomization::GetPresetComboBoxContent)
+				.Font(IDetailLayoutBuilder::GetDetailFont())
+				.ToolTipText(this, &FCameraCropSettingsCustomization::GetPresetComboBoxContent)
+			]
+		];
+}
+
+
+
+void FCameraCropSettingsCustomization::CustomizeChildren(TSharedRef<IPropertyHandle> StructPropertyHandle, class IDetailChildrenBuilder& ChildBuilder, IPropertyTypeCustomizationUtils& StructCustomizationUtils)
+{
+	// Retrieve structure's child properties
+	uint32 NumChildren;
+	StructPropertyHandle->GetNumChildren(NumChildren);
+	TMap<FName, TSharedPtr< IPropertyHandle > > PropertyHandles;
+	for (uint32 ChildIndex = 0; ChildIndex < NumChildren; ++ChildIndex)
+	{
+		TSharedRef<IPropertyHandle> ChildHandle = StructPropertyHandle->GetChildHandle(ChildIndex).ToSharedRef();
+		const FName PropertyName = ChildHandle->GetProperty()->GetFName();
+		PropertyHandles.Add(PropertyName, ChildHandle);
+	}
+	
+	// Retrieve special case properties
+	AspectRatioHandle = PropertyHandles.FindChecked(GET_MEMBER_NAME_CHECKED(FPlateCropSettings, AspectRatio));
+
+	for (auto Iter(PropertyHandles.CreateConstIterator()); Iter; ++Iter)
+	{
+		IDetailPropertyRow& SettingsRow = ChildBuilder.AddProperty(Iter.Value().ToSharedRef());
+	}
+}
+
+TSharedRef<SWidget> FCameraCropSettingsCustomization::MakePresetComboWidget(TSharedPtr<FString> InItem)
+{
+	return
+		SNew(STextBlock)
+		.Text(FText::FromString(*InItem))
+		.Font(IDetailLayoutBuilder::GetDetailFont());
+}
+
+void FCameraCropSettingsCustomization::OnPresetChanged(TSharedPtr<FString> NewSelection, ESelectInfo::Type SelectInfo)
+{
+	// if it's set from code, we did that on purpose
+	if (SelectInfo != ESelectInfo::Direct)
+	{
+		FString const NewPresetName = *NewSelection.Get();
+
+		// search presets for one that matches
+		TArray<FNamedPlateCropPreset> const& Presets = UCineCameraComponent::GetCropPresets();
+		int32 const NumPresets = Presets.Num();
+		for (int32 PresetIdx = 0; PresetIdx < NumPresets; ++PresetIdx)
+		{
+			FNamedPlateCropPreset const& P = Presets[PresetIdx];
+			if (P.Name == NewPresetName)
+			{
+				const FScopedTransaction Transaction(LOCTEXT("ChangeCropPreset", "Change Crop Preset"));
+
+				// copy data from preset into properties
+				// all SetValues except the last set to Interactive so we don't rerun construction scripts and invalidate subsequent property handles
+				ensure(AspectRatioHandle->SetValue(P.CropSettings.AspectRatio, EPropertyValueSetFlags::InteractiveChange | EPropertyValueSetFlags::NotTransactable) == FPropertyAccess::Result::Success);
+				
+				break;
+			}
+		}
+
+		// if none of them found, do nothing
+	}
+}
+
+
+FText FCameraCropSettingsCustomization::GetPresetComboBoxContent() const
+{
+	return FText::FromString(*GetPresetString().Get());
+}
+
+
+TSharedPtr<FString> FCameraCropSettingsCustomization::GetPresetString() const
+{
+  	float AspectRatio;
+ 	AspectRatioHandle->GetValue(AspectRatio);
+ 
+	// search presets for one that matches
+	TArray<FNamedPlateCropPreset> const& Presets = UCineCameraComponent::GetCropPresets();
+	int32 const NumPresets = Presets.Num();
+	for (int32 PresetIdx = 0; PresetIdx < NumPresets; ++PresetIdx)
+	{
+		FNamedPlateCropPreset const& P = Presets[PresetIdx];
+		if (P.CropSettings.AspectRatio == AspectRatio)
+		{
+			// this is the one
+			if (PresetComboList.IsValidIndex(PresetIdx))
+			{
+				return PresetComboList[PresetIdx];
+			}
+		}
+	}
+
+	return PresetComboList.Last();
+}
+
+#undef LOCTEXT_NAMESPACE // CameraCropSettingsCustomization
diff --git a/Engine/Source/Editor/DetailCustomizations/Private/CameraCropSettingsCustomization.h b/Engine/Source/Editor/DetailCustomizations/Private/CameraCropSettingsCustomization.h
new file mode 100644
--- /dev/null
+++ b/Engine/Source/Editor/DetailCustomizations/Private/CameraCropSettingsCustomization.h
@@ -0,0 +1,38 @@
+// Copyright Epic Games, Inc. All Rights Reserved.
+
+#pragma once
+
+#include "CoreMinimal.h"
+#include "Widgets/SWidget.h"
+#include "Widgets/Input/SComboBox.h"
+#include "IPropertyTypeCustomization.h"
+
+class FDetailWidgetRow;
+class IPropertyHandle;
+
+class FCameraCropSettingsCustomization : public IPropertyTypeCustomization
+{
+public:
+	FCameraCropSettingsCustomization();
+
+	static TSharedRef<IPropertyTypeCustomization> MakeInstance();
+
+	/** IPropertyTypeCustomization instance */
+	virtual void CustomizeHeader(TSharedRef<IPropertyHandle> PropertyHandle, FDetailWidgetRow& HeaderRow, IPropertyTypeCustomizationUtils& CustomizationUtils) override;
+	virtual void CustomizeChildren( TSharedRef<IPropertyHandle> StructPropertyHandle, class IDetailChildrenBuilder& ChildBuilder, IPropertyTypeCustomizationUtils& StructCustomizationUtils ) override;
+
+protected:
+
+	TSharedPtr<IPropertyHandle> AspectRatioHandle;
+
+	TSharedPtr<class SComboBox< TSharedPtr<FString> > > PresetComboBox;
+	TArray< TSharedPtr< FString > >						PresetComboList;
+
+	void OnPresetChanged(TSharedPtr<FString> NewSelection, ESelectInfo::Type SelectInfo);
+	TSharedRef<SWidget> MakePresetComboWidget(TSharedPtr<FString> InItem);
+
+	FText GetPresetComboBoxContent() const;
+	TSharedPtr<FString> GetPresetString() const;
+
+protected:
+};
diff --git a/Engine/Source/Editor/DetailCustomizations/Private/CameraLensSettingsCustomization.cpp b/Engine/Source/Editor/DetailCustomizations/Private/CameraLensSettingsCustomization.cpp
--- a/Engine/Source/Editor/DetailCustomizations/Private/CameraLensSettingsCustomization.cpp
+++ b/Engine/Source/Editor/DetailCustomizations/Private/CameraLensSettingsCustomization.cpp
@@ -80,6 +80,7 @@ void FCameraLensSettingsCustomization::CustomizeChildren(TSharedRef<IPropertyHan
 	MinFStopHandle = PropertyHandles.FindChecked(GET_MEMBER_NAME_CHECKED(FCameraLensSettings, MinFStop));
 	MaxFStopHandle = PropertyHandles.FindChecked(GET_MEMBER_NAME_CHECKED(FCameraLensSettings, MaxFStop));
 	MinFocusDistanceHandle = PropertyHandles.FindChecked(GET_MEMBER_NAME_CHECKED(FCameraLensSettings, MinimumFocusDistance));
+	SqueezeFactorHandle = PropertyHandles.FindChecked(GET_MEMBER_NAME_CHECKED(FCameraLensSettings, SqueezeFactor));
 	DiaphragmBladeCountHandle = PropertyHandles.FindChecked(GET_MEMBER_NAME_CHECKED(FCameraLensSettings, DiaphragmBladeCount));
 
 	for (auto Iter(PropertyHandles.CreateConstIterator()); Iter; ++Iter)
@@ -126,6 +127,7 @@ void FCameraLensSettingsCustomization::OnPresetChanged(TSharedPtr<FString> NewSe
 				ensure(MinFStopHandle->SetValue(P.LensSettings.MinFStop, EPropertyValueSetFlags::InteractiveChange | EPropertyValueSetFlags::NotTransactable) == FPropertyAccess::Result::Success);
 				ensure(MaxFStopHandle->SetValue(P.LensSettings.MaxFStop, EPropertyValueSetFlags::InteractiveChange | EPropertyValueSetFlags::NotTransactable) == FPropertyAccess::Result::Success);
 				ensure(MinFocusDistanceHandle->SetValue(P.LensSettings.MinimumFocusDistance, EPropertyValueSetFlags::InteractiveChange | EPropertyValueSetFlags::NotTransactable) == FPropertyAccess::Result::Success);
+				ensure(SqueezeFactorHandle->SetValue(P.LensSettings.SqueezeFactor, EPropertyValueSetFlags::InteractiveChange | EPropertyValueSetFlags::NotTransactable) == FPropertyAccess::Result::Success);
 				ensure(DiaphragmBladeCountHandle->SetValue(P.LensSettings.DiaphragmBladeCount, EPropertyValueSetFlags::InteractiveChange | EPropertyValueSetFlags::NotTransactable) == FPropertyAccess::Result::Success);
 
 				break;
@@ -168,6 +170,9 @@ TSharedPtr<FString> FCameraLensSettingsCustomization::GetPresetString() const
 	float MinFocusDistance;
 	MinFocusDistanceHandle->GetValue(MinFocusDistance);
 
+	float SqueezeFactor;
+	SqueezeFactorHandle->GetValue(SqueezeFactor);
+
 	int32 DiaphragmBladeCount;
 	DiaphragmBladeCountHandle->GetValue(DiaphragmBladeCount);
 
@@ -178,7 +183,7 @@ TSharedPtr<FString> FCameraLensSettingsCustomization::GetPresetString() const
 	{
 		FNamedLensPreset const& P = Presets[PresetIdx];
 		if ((P.LensSettings.MinFocalLength == MinFocalLength) && (P.LensSettings.MaxFocalLength == MaxFocalLength) && (P.LensSettings.MinFStop == MinFStop) && 
-			(P.LensSettings.MaxFStop == MaxFStop) && (P.LensSettings.MinimumFocusDistance == MinFocusDistance) && (P.LensSettings.DiaphragmBladeCount == DiaphragmBladeCount) )
+			(P.LensSettings.MaxFStop == MaxFStop) && (P.LensSettings.MinimumFocusDistance == MinFocusDistance) && (P.LensSettings.SqueezeFactor == SqueezeFactor) && (P.LensSettings.DiaphragmBladeCount == DiaphragmBladeCount) )
 		{
 			// this is the one
 			if (PresetComboList.IsValidIndex(PresetIdx + 1))
diff --git a/Engine/Source/Editor/DetailCustomizations/Private/CameraLensSettingsCustomization.h b/Engine/Source/Editor/DetailCustomizations/Private/CameraLensSettingsCustomization.h
--- a/Engine/Source/Editor/DetailCustomizations/Private/CameraLensSettingsCustomization.h
+++ b/Engine/Source/Editor/DetailCustomizations/Private/CameraLensSettingsCustomization.h
@@ -28,7 +28,7 @@ protected:
 	TSharedPtr<IPropertyHandle> MinFStopHandle;
 	TSharedPtr<IPropertyHandle> MaxFStopHandle;
 	TSharedPtr<IPropertyHandle> MinFocusDistanceHandle;
-	TSharedPtr<IPropertyHandle> MaxReproRatioHandle;
+	TSharedPtr<IPropertyHandle> SqueezeFactorHandle;
 	TSharedPtr<IPropertyHandle> DiaphragmBladeCountHandle;
 
 	TSharedPtr<class SComboBox< TSharedPtr<FString> > > PresetComboBox;
diff --git a/Engine/Source/Editor/DetailCustomizations/Private/DetailCustomizations.cpp b/Engine/Source/Editor/DetailCustomizations/Private/DetailCustomizations.cpp
--- a/Engine/Source/Editor/DetailCustomizations/Private/DetailCustomizations.cpp
+++ b/Engine/Source/Editor/DetailCustomizations/Private/DetailCustomizations.cpp
@@ -23,6 +23,7 @@
 #include "BodySetupDetails.h"
 #include "BoundsCopyComponentDetails.h"
 #include "BrushDetails.h"
+#include "CameraCropSettingsCustomization.h"
 #include "CameraDetails.h"
 #include "CameraFilmbackSettingsCustomization.h"
 #include "CameraFocusSettingsCustomization.h"
@@ -311,6 +312,7 @@ void FDetailCustomizationsModule::RegisterPropertyTypeCustomizations()
 	RegisterCustomPropertyTypeLayout("CompositeReroute", FOnGetPropertyTypeCustomizationInstance::CreateStatic(&FCompositeRerouteCustomization::MakeInstance));
 	RegisterCustomPropertyTypeLayout("CameraFilmbackSettings", FOnGetPropertyTypeCustomizationInstance::CreateStatic(&FCameraFilmbackSettingsCustomization::MakeInstance));
 	RegisterCustomPropertyTypeLayout("CameraLensSettings", FOnGetPropertyTypeCustomizationInstance::CreateStatic(&FCameraLensSettingsCustomization::MakeInstance));
+	RegisterCustomPropertyTypeLayout("PlateCropSettings", FOnGetPropertyTypeCustomizationInstance::CreateStatic(&FCameraCropSettingsCustomization::MakeInstance));
 	RegisterCustomPropertyTypeLayout("CameraFocusSettings", FOnGetPropertyTypeCustomizationInstance::CreateStatic(&FCameraFocusSettingsCustomization::MakeInstance));
 	RegisterCustomPropertyTypeLayout("MovieSceneSequenceLoopCount", FOnGetPropertyTypeCustomizationInstance::CreateStatic(&FMovieSceneSequenceLoopCountCustomization::MakeInstance));
 	RegisterCustomPropertyTypeLayout("MovieSceneBindingOverrideData", FOnGetPropertyTypeCustomizationInstance::CreateStatic(&FMovieSceneBindingOverrideDataCustomization::MakeInstance));
diff --git a/Engine/Source/Runtime/CinematicCamera/Private/CineCameraComponent.cpp b/Engine/Source/Runtime/CinematicCamera/Private/CineCameraComponent.cpp
--- a/Engine/Source/Runtime/CinematicCamera/Private/CineCameraComponent.cpp
+++ b/Engine/Source/Runtime/CinematicCamera/Private/CineCameraComponent.cpp
@@ -43,10 +43,11 @@ UCineCameraComponent::UCineCameraComponent()
 	// Null check here because CDO has no Archetype
 	if (Template)
 	{
-		// default filmback
+		// default filmback, lens and crop
 		SetFilmbackPresetByNameInternal(Template->DefaultFilmbackPreset, Filmback);
 		SetFilmbackPresetByNameInternal(Template->DefaultFilmbackPresetName_DEPRECATED, FilmbackSettings_DEPRECATED);
 		SetLensPresetByNameInternal(Template->DefaultLensPresetName);
+		SetCropPresetByNameInternal(Template->DefaultCropPresetName);
 		// other lens defaults
 		CurrentAperture = Template->DefaultLensFStop;
 		CurrentFocalLength = Template->DefaultLensFocalLength;
@@ -249,7 +250,17 @@ void UCineCameraComponent::SetFieldOfView(float InFieldOfView)
 {
 	Super::SetFieldOfView(InFieldOfView);
 
-	CurrentFocalLength = (Filmback.SensorWidth / 2.f) / FMath::Tan(FMath::DegreesToRadians(InFieldOfView / 2.f));
+	float CropedSensorWidth = Filmback.SensorWidth * LensSettings.SqueezeFactor;
+	if (CropSettings.AspectRatio > 0.0f)
+	{
+		float DesqueezeAspectRatio = Filmback.SensorWidth * LensSettings.SqueezeFactor / Filmback.SensorHeight;
+		if (CropSettings.AspectRatio < DesqueezeAspectRatio)
+		{
+			CropedSensorWidth *= CropSettings.AspectRatio / DesqueezeAspectRatio;
+		}
+	}
+
+	CurrentFocalLength = (CropedSensorWidth / 2.f) / FMath::Tan(FMath::DegreesToRadians(InFieldOfView / 2.f));
 }
 
 void UCineCameraComponent::SetCurrentFocalLength(float InFocalLength)
@@ -260,16 +271,42 @@ void UCineCameraComponent::SetCurrentFocalLength(float InFocalLength)
 
 float UCineCameraComponent::GetHorizontalFieldOfView() const
 {
-	return (CurrentFocalLength > 0.f)
-		? FMath::RadiansToDegrees(2.f * FMath::Atan(Filmback.SensorWidth / (2.f * CurrentFocalLength)))
-		: 0.f;
+	if (CurrentFocalLength > 0.f)
+	{
+		float CropedSensorWidth = Filmback.SensorWidth * LensSettings.SqueezeFactor;
+		if (CropSettings.AspectRatio > 0.0f)
+		{
+			float DesqueezeAspectRatio = Filmback.SensorWidth * LensSettings.SqueezeFactor / Filmback.SensorHeight;
+			if (CropSettings.AspectRatio < DesqueezeAspectRatio)
+			{
+				CropedSensorWidth *= CropSettings.AspectRatio / DesqueezeAspectRatio;
+			}
+		}
+
+		return FMath::RadiansToDegrees(2.f * FMath::Atan(CropedSensorWidth / (2.f * CurrentFocalLength)));
+	}
+
+	return 0.f;
 }
 
 float UCineCameraComponent::GetVerticalFieldOfView() const
 {
-	return (CurrentFocalLength > 0.f)
-		? FMath::RadiansToDegrees(2.f * FMath::Atan(Filmback.SensorHeight / (2.f * CurrentFocalLength)))
-		: 0.f;
+	if (CurrentFocalLength > 0.f)
+	{
+		float CropedSensorHeight = Filmback.SensorHeight;
+		if (CropSettings.AspectRatio > 0.0f)
+		{
+			float DesqueezeAspectRatio = Filmback.SensorWidth * LensSettings.SqueezeFactor / Filmback.SensorHeight;
+			if (DesqueezeAspectRatio < CropSettings.AspectRatio)
+			{
+				CropedSensorHeight *= DesqueezeAspectRatio / CropSettings.AspectRatio;
+			}
+		}
+
+		return FMath::RadiansToDegrees(2.f * FMath::Atan(CropedSensorHeight / (2.f * CurrentFocalLength)));
+	}
+
+	return 0.f;
 }
 
 FString UCineCameraComponent::GetFilmbackPresetName() const
@@ -348,6 +385,44 @@ void UCineCameraComponent::SetLensPresetByNameInternal(const FString& InPresetNa
 	}
 }
 
+FString UCineCameraComponent::GetCropPresetName() const
+{
+	TArray<FNamedPlateCropPreset> const& Presets = UCineCameraComponent::GetCropPresets();
+	int32 const NumPresets = Presets.Num();
+	for (int32 PresetIdx = 0; PresetIdx < NumPresets; ++PresetIdx)
+	{
+		FNamedPlateCropPreset const& P = Presets[PresetIdx];
+		if (P.CropSettings == CropSettings)
+		{
+			return P.Name;
+		}
+	}
+
+	return FString();
+}
+
+void UCineCameraComponent::SetCropPresetByName(const FString& InPresetName)
+{
+	SetCropPresetByNameInternal(InPresetName);
+	// Explicitely call RecalcDerivedData() when invoked via Blueprint, since no other method (incl. PostEditChangeProperty) will trigger
+	RecalcDerivedData();
+}
+
+void UCineCameraComponent::SetCropPresetByNameInternal(const FString& InPresetName)
+{
+	TArray<FNamedPlateCropPreset> const& Presets = UCineCameraComponent::GetCropPresets();
+	int32 const NumPresets = Presets.Num();
+	for (int32 PresetIdx = 0; PresetIdx < NumPresets; ++PresetIdx)
+	{
+		FNamedPlateCropPreset const& P = Presets[PresetIdx];
+		if (P.Name == InPresetName)
+		{
+			CropSettings = P.CropSettings;
+			break;
+		}
+	}
+}
+
 float UCineCameraComponent::GetWorldToMetersScale() const
 {
 	UWorld const* const World = GetWorld();
@@ -379,6 +454,12 @@ TArray<FNamedLensPreset> const& UCineCameraComponent::GetLensPresets()
 	return GetDefault<UCineCameraComponent>()->LensPresets;
 }
 
+// static
+TArray<FNamedPlateCropPreset> const& UCineCameraComponent::GetCropPresets()
+{
+	return GetDefault<UCineCameraComponent>()->CropPresets;
+}
+
 void UCineCameraComponent::RecalcDerivedData()
 {
 	// validate incorrect values
@@ -393,7 +474,11 @@ void UCineCameraComponent::RecalcDerivedData()
 
 	FieldOfView = GetHorizontalFieldOfView();
 	Filmback.SensorAspectRatio = (Filmback.SensorHeight > 0.f) ? (Filmback.SensorWidth / Filmback.SensorHeight) : 0.f;
-	AspectRatio = Filmback.SensorAspectRatio;
+	AspectRatio = Filmback.SensorAspectRatio * LensSettings.SqueezeFactor;
+	if (CropSettings.AspectRatio > 0.0f)
+	{
+		AspectRatio = CropSettings.AspectRatio;
+	}
 
 #if WITH_EDITORONLY_DATA
 	CurrentHorizontalFOV = FieldOfView;			// informational variable only, for editor users
@@ -477,22 +562,23 @@ FText UCineCameraComponent::GetFilmbackText() const
 	if (Preset)
 	{
 		return FText::Format(
-			LOCTEXT("PresetFormat","FilmbackPreset: {0} | Zoom: {1}mm | Av: {2}"),
+			LOCTEXT("PresetFormat","FilmbackPreset: {0} | Zoom: {1}mm | Av: {2} | Squeeze: {3}"),
 			FText::FromString(Preset->Name),
 			FText::AsNumber(CurrentFocalLength),
-			FText::AsNumber(CurrentAperture)
+			FText::AsNumber(CurrentAperture),
+			FText::AsNumber(LensSettings.SqueezeFactor)
 		);
 	}
 	else
 	{
 		FNumberFormattingOptions Opts = FNumberFormattingOptions().SetMaximumFractionalDigits(1);
 		return FText::Format(
-			LOCTEXT("CustomFilmbackFormat", "Custom ({0}mm x {1}mm) | Zoom: {2}mm | Av: {3}"),
+			LOCTEXT("CustomFilmbackFormat", "Custom ({0}mm x {1}mm) | Zoom: {2}mm | Av: {3} | Squeeze: {4}"),
 			FText::AsNumber(SensorWidth, &Opts),
 			FText::AsNumber(SensorHeight, &Opts),
 			FText::AsNumber(CurrentFocalLength),
-			FText::AsNumber(CurrentAperture)
-
+			FText::AsNumber(CurrentAperture),
+			FText::AsNumber(LensSettings.SqueezeFactor)
 		);
 	}
 }
@@ -524,6 +610,7 @@ void UCineCameraComponent::UpdateCameraLens(float DeltaTime, FMinimalViewInfo& D
 		DesiredView.PostProcessSettings.bOverride_DepthOfFieldBladeCount = false;
 		DesiredView.PostProcessSettings.bOverride_DepthOfFieldFocalDistance = false;
 		DesiredView.PostProcessSettings.bOverride_DepthOfFieldSensorWidth = false;
+		DesiredView.PostProcessSettings.bOverride_DepthOfFieldSqueezeFactor = false;
 	}
 	else if (FocusSettings.FocusMethod == ECameraFocusMethod::Disable)
 	{
@@ -566,6 +653,9 @@ void UCineCameraComponent::UpdateCameraLens(float DeltaTime, FMinimalViewInfo& D
 
 		DesiredView.PostProcessSettings.bOverride_DepthOfFieldSensorWidth = true;
 		DesiredView.PostProcessSettings.DepthOfFieldSensorWidth = Filmback.SensorWidth;
+
+		DesiredView.PostProcessSettings.bOverride_DepthOfFieldSqueezeFactor = true;
+		DesiredView.PostProcessSettings.DepthOfFieldSqueezeFactor = LensSettings.SqueezeFactor;
 	}
 }
 
diff --git a/Engine/Source/Runtime/CinematicCamera/Public/CineCameraComponent.h b/Engine/Source/Runtime/CinematicCamera/Public/CineCameraComponent.h
--- a/Engine/Source/Runtime/CinematicCamera/Public/CineCameraComponent.h
+++ b/Engine/Source/Runtime/CinematicCamera/Public/CineCameraComponent.h
@@ -103,6 +103,10 @@ struct FCameraLensSettings
 	UPROPERTY(EditAnywhere, BlueprintReadWrite, Category = "Lens", meta = (ForceUnits = mm))
 	float MinimumFocusDistance = 0.f;
 
+	/** Squeeze factor for anamorphic lenses. */
+	UPROPERTY(EditAnywhere, BlueprintReadWrite, Category = "Lens", meta = (ClampMin = "1", ClampMax = "2"))
+	float SqueezeFactor = 1.f;
+
 	/** Number of blades of diaphragm. */
 	UPROPERTY(EditAnywhere, BlueprintReadWrite, Category = "Lens", meta = (ClampMin = "4", ClampMax = "16"))
 	int32 DiaphragmBladeCount = 0;
@@ -114,6 +118,7 @@ struct FCameraLensSettings
 			&& (MinFStop == Other.MinFStop)
 			&& (MaxFStop == Other.MaxFStop)
 			&& (MinimumFocusDistance == Other.MinimumFocusDistance)
+			&& (SqueezeFactor == Other.SqueezeFactor)
 			&& (DiaphragmBladeCount == Other.DiaphragmBladeCount);
 	}
 };
@@ -132,6 +137,37 @@ struct FNamedLensPreset
 	FCameraLensSettings LensSettings;
 };
 
+/** 
+ * #note, this struct has a details customization in CameraCropSettingsCustomization.cpp/h
+ */
+USTRUCT(BlueprintType)
+struct FPlateCropSettings
+{
+	GENERATED_USTRUCT_BODY()
+
+	UPROPERTY(EditAnywhere, BlueprintReadWrite, Category = "Crop", meta=(DisplayName = "Cropped Aspect Ratio"))
+	float AspectRatio = 0.f;
+
+	bool operator==(const FPlateCropSettings& Other) const
+	{
+		return (AspectRatio == Other.AspectRatio);
+	}
+};
+
+/** A named bundle of crop settings used to implement crop presets. */
+USTRUCT(BlueprintType)
+struct FNamedPlateCropPreset
+{
+	GENERATED_USTRUCT_BODY()
+
+	/** Name for the preset. */
+	UPROPERTY(BlueprintReadWrite, Category = "Crop")
+	FString Name;
+
+	UPROPERTY(BlueprintReadWrite, Category = "Crop")
+	FPlateCropSettings CropSettings;
+};
+
 /** Supported methods for focusing the camera. */
 UENUM()
 enum class ECameraFocusMethod : uint8
@@ -260,6 +296,10 @@ public:
 	UPROPERTY(EditAnywhere, BlueprintReadWrite, Category = "Current Camera Settings")
 	FCameraFocusSettings FocusSettings;
 
+	/** Controls the crop settings. */
+	UPROPERTY(EditAnywhere, BlueprintReadWrite, Category = "Current Camera Settings")
+	FPlateCropSettings CropSettings;
+
 	/** Current focal length of the camera (i.e. controls FoV, zoom) */
 	UPROPERTY(Interp, BlueprintSetter = SetCurrentFocalLength, EditAnywhere, BlueprintReadWrite, Category = "Current Camera Settings")
 	float CurrentFocalLength;
@@ -318,6 +358,14 @@ public:
 	/** Set the current lens settings by preset name. */
 	UFUNCTION(BlueprintCallable, Category = "Cine Camera")
 	void SetLensPresetByName(const FString& InPresetName);
+	
+	/** Returns the lens name of the camera with the current settings. */
+	UFUNCTION(BlueprintCallable, Category = "Cine Camera")
+	FString GetCropPresetName() const;
+
+	/** Set the current lens settings by preset name. */
+	UFUNCTION(BlueprintCallable, Category = "Cine Camera")
+	void SetCropPresetByName(const FString& InPresetName);
 
 	/** Returns a copy of the list of available filmback presets. */
 	UFUNCTION(BlueprintCallable, Category = "Cine Camera")
@@ -333,6 +381,9 @@ public:
 	/** Returns a list of available lens presets. */
 	static TArray<FNamedLensPreset> const& GetLensPresets();
 
+	/** Returns a list of available crop presets. */
+	static TArray<FNamedPlateCropPreset> const& GetCropPresets();
+
 #if WITH_EDITOR
 	/** Update the debug focus plane position and orientation. */
 	void UpdateDebugFocusPlane();
@@ -392,6 +443,10 @@ protected:
 	/** List of available lens presets */
 	UPROPERTY(config)
 	TArray<FNamedLensPreset> LensPresets;
+	
+	/** List of available lens presets */
+	UPROPERTY(config)
+	TArray<FNamedPlateCropPreset> CropPresets;
 
 	/** Deprecated. See DefaultFilmbackPreset */
 	UPROPERTY(config)
@@ -404,7 +459,11 @@ protected:
 	/** Name of the default lens preset */
 	UPROPERTY(config)
 	FString DefaultLensPresetName;
-	
+
+	/** Name of the default crop preset */
+	UPROPERTY(config)
+	FString DefaultCropPresetName;
+
 	/** Default focal length (will be constrained by default lens) */
 	UPROPERTY(config)
 	float DefaultLensFocalLength;
@@ -423,6 +482,7 @@ private:
 	float GetDesiredFocusDistance(const FVector& InLocation) const;
 	void SetLensPresetByNameInternal(const FString& InPresetName);
 	void SetFilmbackPresetByNameInternal(const FString& InPresetName, FCameraFilmbackSettings& InOutFilmbackSettings);
+	void SetCropPresetByNameInternal(const FString& InPresetName);
 
 #if WITH_EDITORONLY_DATA
 	void CreateDebugFocusPlane();
diff --git a/Engine/Source/Runtime/Engine/Classes/Engine/Scene.h b/Engine/Source/Runtime/Engine/Classes/Engine/Scene.h
--- a/Engine/Source/Runtime/Engine/Classes/Engine/Scene.h
+++ b/Engine/Source/Runtime/Engine/Classes/Engine/Scene.h
@@ -1115,6 +1115,9 @@ struct FPostProcessSettings
 
 	UPROPERTY(EditAnywhere, BlueprintReadWrite, Category=Overrides, meta=(PinHiddenByDefault, InlineEditConditionToggle))
 	uint8 bOverride_DepthOfFieldSensorWidth:1;
+	
+	UPROPERTY(EditAnywhere, BlueprintReadWrite, Category=Overrides, meta=(PinHiddenByDefault, InlineEditConditionToggle))
+	uint8 bOverride_DepthOfFieldSqueezeFactor:1;
 
 	UPROPERTY(EditAnywhere, BlueprintReadWrite, Category=Overrides, meta=(PinHiddenByDefault, InlineEditConditionToggle))
 	uint8 bOverride_DepthOfFieldDepthBlurRadius:1;
@@ -1980,6 +1983,10 @@ struct FPostProcessSettings
 	UPROPERTY(BlueprintReadWrite, Category="Lens|Depth of Field", meta=(ForceUnits=mm, ClampMin = "0.1", UIMin="0.1", UIMax= "1000.0", editcondition = "bOverride_DepthOfFieldSensorWidth", DisplayName = "Sensor Width (mm)"))
 	float DepthOfFieldSensorWidth;
 
+	/** Width of the camera sensor to assume, in mm. */
+	UPROPERTY(BlueprintReadWrite, Category = "Lens|Depth of Field", meta = (ClampMin = "1.0", ClampMax = "2.0", editcondition = "bOverride_DepthOfFieldSqueezeFactor", DisplayName = "Squeeze Factor"))
+	float DepthOfFieldSqueezeFactor;
+
 	/** Distance in which the Depth of Field effect should be sharp, in unreal units (cm) */
 	UPROPERTY(interp, BlueprintReadWrite, Category="Lens|Depth of Field", meta=(ClampMin = "0.0", UIMin = "1.0", UIMax = "10000.0", editcondition = "bOverride_DepthOfFieldFocalDistance", DisplayName = "Focal Distance"))
 	float DepthOfFieldFocalDistance;
diff --git a/Engine/Source/Runtime/Engine/Private/Scene.cpp b/Engine/Source/Runtime/Engine/Private/Scene.cpp
--- a/Engine/Source/Runtime/Engine/Private/Scene.cpp
+++ b/Engine/Source/Runtime/Engine/Private/Scene.cpp
@@ -548,6 +548,7 @@ FPostProcessSettings::FPostProcessSettings()
 	DepthOfFieldMinFstop = 1.2f;
 	DepthOfFieldBladeCount = FPostProcessSettings::kDefaultDepthOfFieldBladeCount;
 	DepthOfFieldSensorWidth = 24.576f;			// APS-C
+	DepthOfFieldSqueezeFactor = 1.0f;
 	DepthOfFieldDepthBlurAmount = 1.0f;
 	DepthOfFieldDepthBlurRadius = 0.0f;
 	DepthOfFieldFocalRegion = 0.0f;
@@ -752,6 +753,7 @@ FPostProcessSettings::FPostProcessSettings(const FPostProcessSettings& Settings)
 	, bOverride_DepthOfFieldMinFstop(Settings.bOverride_DepthOfFieldMinFstop)
 	, bOverride_DepthOfFieldBladeCount(Settings.bOverride_DepthOfFieldBladeCount)
 	, bOverride_DepthOfFieldSensorWidth(Settings.bOverride_DepthOfFieldSensorWidth)
+	, bOverride_DepthOfFieldSqueezeFactor(Settings.bOverride_DepthOfFieldSqueezeFactor)
 	, bOverride_DepthOfFieldDepthBlurRadius(Settings.bOverride_DepthOfFieldDepthBlurRadius)
 	, bOverride_DepthOfFieldDepthBlurAmount(Settings.bOverride_DepthOfFieldDepthBlurAmount)
 	, bOverride_DepthOfFieldFocalRegion(Settings.bOverride_DepthOfFieldFocalRegion)
@@ -965,6 +967,7 @@ FPostProcessSettings::FPostProcessSettings(const FPostProcessSettings& Settings)
 	, ColorGradingIntensity(Settings.ColorGradingIntensity)
 	, ColorGradingLUT(Settings.ColorGradingLUT)
 	, DepthOfFieldSensorWidth(Settings.DepthOfFieldSensorWidth)
+	, DepthOfFieldSqueezeFactor(Settings.DepthOfFieldSqueezeFactor)
 	, DepthOfFieldFocalDistance(Settings.DepthOfFieldFocalDistance)
 	, DepthOfFieldDepthBlurAmount(Settings.DepthOfFieldDepthBlurAmount)
 	, DepthOfFieldDepthBlurRadius(Settings.DepthOfFieldDepthBlurRadius)
diff --git a/Engine/Source/Runtime/Engine/Private/SceneView.cpp b/Engine/Source/Runtime/Engine/Private/SceneView.cpp
--- a/Engine/Source/Runtime/Engine/Private/SceneView.cpp
+++ b/Engine/Source/Runtime/Engine/Private/SceneView.cpp
@@ -1537,6 +1537,7 @@ void FSceneView::OverridePostProcessSettings(const FPostProcessSettings& Src, fl
 		LERP_PP(DepthOfFieldFstop);
 		LERP_PP(DepthOfFieldMinFstop);
 		LERP_PP(DepthOfFieldSensorWidth);
+		LERP_PP(DepthOfFieldSqueezeFactor);
 		LERP_PP(DepthOfFieldDepthBlurRadius);
 		LERP_PP(DepthOfFieldDepthBlurAmount);
 		LERP_PP(DepthOfFieldFocalRegion);
diff --git a/Engine/Source/Runtime/Renderer/Private/PostProcess/DiaphragmDOF.cpp b/Engine/Source/Runtime/Renderer/Private/PostProcess/DiaphragmDOF.cpp
--- a/Engine/Source/Runtime/Renderer/Private/PostProcess/DiaphragmDOF.cpp
+++ b/Engine/Source/Runtime/Renderer/Private/PostProcess/DiaphragmDOF.cpp
@@ -617,8 +617,13 @@ END_SHADER_PARAMETER_STRUCT()
 
 
 BEGIN_SHADER_PARAMETER_STRUCT(FDOFCocModelShaderParameters, )
-	SHADER_PARAMETER(FVector4f, CocModelParameters)
-	SHADER_PARAMETER(FVector2f, DepthBlurParameters)
+	SHADER_PARAMETER(float, CocInfinityRadius)
+	SHADER_PARAMETER(float, CocMinRadius)
+	SHADER_PARAMETER(float, CocMaxRadius)
+	SHADER_PARAMETER(float, CocSqueeze)
+	SHADER_PARAMETER(float, CocInvSqueeze)
+	SHADER_PARAMETER(float, DepthBlurRadius)
+	SHADER_PARAMETER(float, DepthBlurExponent)
 END_SHADER_PARAMETER_STRUCT()
 
 void SetCocModelParameters(
@@ -626,11 +631,13 @@ void SetCocModelParameters(
 	const DiaphragmDOF::FPhysicalCocModel& CocModel,
 	float CocRadiusBasis = 1.0f)
 {
-	OutParameters->CocModelParameters.X = CocRadiusBasis * CocModel.InfinityBackgroundCocRadius;
-	OutParameters->CocModelParameters.Y = CocRadiusBasis * CocModel.MinForegroundCocRadius;
-	OutParameters->CocModelParameters.Z = CocRadiusBasis * CocModel.MaxBackgroundCocRadius;
-	OutParameters->DepthBlurParameters.X = CocModel.DepthBlurExponent;
-	OutParameters->DepthBlurParameters.Y = CocRadiusBasis * CocModel.MaxDepthBlurRadius;
+	OutParameters->CocInfinityRadius = CocRadiusBasis * CocModel.InfinityBackgroundCocRadius;
+	OutParameters->CocMinRadius = CocRadiusBasis * CocModel.MinForegroundCocRadius;
+	OutParameters->CocMaxRadius = CocRadiusBasis * CocModel.MaxBackgroundCocRadius;
+	OutParameters->CocSqueeze = CocModel.Squeeze;
+	OutParameters->CocInvSqueeze = 1.0f / CocModel.Squeeze;
+	OutParameters->DepthBlurRadius = CocRadiusBasis * CocModel.MaxDepthBlurRadius;
+	OutParameters->DepthBlurExponent = CocModel.DepthBlurExponent;
 }
 
 
@@ -861,6 +868,7 @@ class FDiaphragmDOFReduceCS : public FDiaphragmDOFShader
 		SHADER_PARAMETER(float, PreProcessingToProcessingCocRadiusFactor)
 		SHADER_PARAMETER(float, MinScatteringCocRadius)
 		SHADER_PARAMETER(float, NeighborCompareMaxColor)
+		SHADER_PARAMETER(float, CocSqueeze)
 		
 		SHADER_PARAMETER_RDG_TEXTURE(Texture2D, EyeAdaptationTexture)
 		SHADER_PARAMETER_STRUCT_INCLUDE(FDOFCommonShaderParameters, CommonParameters)
@@ -1082,6 +1090,8 @@ class FDiaphragmDOFGatherCS : public FDiaphragmDOFShader
 		SHADER_PARAMETER(FVector2f, InputBufferUVToOutputPixel)
 		SHADER_PARAMETER(float, MipBias)
 		SHADER_PARAMETER(float, MaxRecombineAbsCocRadius)
+		SHADER_PARAMETER(float, CocSqueeze)
+		SHADER_PARAMETER(float, CocInvSqueeze)
 
 		SHADER_PARAMETER_STRUCT_INCLUDE(FDOFTileDecisionParameters, TileDecisionParameters)
 		SHADER_PARAMETER_STRUCT_INCLUDE(FDOFCommonShaderParameters, CommonParameters)
@@ -1147,6 +1157,8 @@ BEGIN_SHADER_PARAMETER_STRUCT(FDOFHybridScatterParameters, )
 	SHADER_PARAMETER(FVector4f, ViewportSize)
 	SHADER_PARAMETER(float, CocRadiusToCircumscribedRadius)
 	SHADER_PARAMETER(float, ScatteringScaling)
+	SHADER_PARAMETER(float, CocSqueeze)
+	SHADER_PARAMETER(float, CocInvSqueeze)
 	
 	SHADER_PARAMETER_STRUCT_INCLUDE(FDOFCommonShaderParameters, CommonParameters)
 	
@@ -1949,6 +1961,7 @@ FRDGTextureRef DiaphragmDOF::AddPasses(
 			PassParameters->PreProcessingToProcessingCocRadiusFactor = PreProcessingToProcessingCocRadiusFactor;
 			PassParameters->MinScatteringCocRadius = MinScatteringCocRadius;
 			PassParameters->NeighborCompareMaxColor = CVarScatterNeighborCompareMaxColor.GetValueOnRenderThread();
+			PassParameters->CocSqueeze = CocModel.Squeeze;
 			
 			PassParameters->EyeAdaptationTexture = GetEyeAdaptationTexture(GraphBuilder, View);
 			PassParameters->CommonParameters = CommonParameters;
@@ -2280,6 +2293,8 @@ FRDGTextureRef DiaphragmDOF::AddPasses(
 				float(SrcSize.Y * GatheringViewSize.Y) / float(PreprocessViewSize.Y));
 			PassParameters->MipBias = FMath::Log2(float(PreprocessViewSize.X) / float(GatheringViewSize.X));
 			PassParameters->MaxRecombineAbsCocRadius = float(kMaxSlightOutOfFocusCocRadius) / PreProcessingToProcessingCocRadiusFactor;
+			PassParameters->CocSqueeze = CocModel.Squeeze;
+			PassParameters->CocInvSqueeze = 1.0f / CocModel.Squeeze;
 
 			PassParameters->TileDecisionParameters.MinGatherRadius = PassParameters->MaxRecombineAbsCocRadius - 1;
 			PassParameters->TileDecisionParameters.SlightOutOfFocusRadiusBoundary = float(kMaxSlightOutOfFocusCocRadius) / PreProcessingToProcessingCocRadiusFactor;
@@ -2424,6 +2439,8 @@ FRDGTextureRef DiaphragmDOF::AddPasses(
 			PassParameters->ViewportSize = FVector4f(GatheringViewSize.X, GatheringViewSize.Y, 1.0f / GatheringViewSize.X, 1.0f / GatheringViewSize.Y);
 			PassParameters->CocRadiusToCircumscribedRadius = BokehModel.CocRadiusToCircumscribedRadius;
 			PassParameters->ScatteringScaling = float(GatheringViewSize.X) / float(PreprocessViewSize.X);
+			PassParameters->CocSqueeze = CocModel.Squeeze;
+			PassParameters->CocInvSqueeze = 1.0 / CocModel.Squeeze;
 			PassParameters->CommonParameters = CommonParameters;
 			if (bEnableScatterBokehSettings)
 				PassParameters->BokehLUT = ScatteringBokehLUT;
diff --git a/Engine/Source/Runtime/Renderer/Private/PostProcess/DiaphragmDOF.h b/Engine/Source/Runtime/Renderer/Private/PostProcess/DiaphragmDOF.h
--- a/Engine/Source/Runtime/Renderer/Private/PostProcess/DiaphragmDOF.h
+++ b/Engine/Source/Runtime/Renderer/Private/PostProcess/DiaphragmDOF.h
@@ -40,6 +40,9 @@ struct FPhysicalCocModel
 	// Focus distance.
 	float FocusDistance;
 
+	// SqueezeFactor = VerticalFocalDistance / HorizontalFocalDistance
+	float Squeeze;
+
 	// The maximum radius of depth blur.
 	float MaxDepthBlurRadius;
 	float DepthBlurExponent;
@@ -49,17 +52,17 @@ struct FPhysicalCocModel
 	
 	/** Returns the CocRadius in half res pixels for given scene depth (in world unit).
 	 *
-	 * Notes: Matches Engine/Shaders/Private/DiaphragmDOF/Common.ush's DepthToHalfResCocRadius().
+	 * Notes: Matches Engine/Shaders/Private/DiaphragmDOF/Common.ush's SceneDepthToCocRadius().
 	 */
 	float DepthToResCocRadius(float SceneDepth, float HorizontalResolution) const;
 
-	/** Returns limit(DepthToHalfResCocRadius) for SceneDepth -> Infinity. */
+	/** Returns limit(SceneDepthToCocRadius) for SceneDepth -> Infinity. */
 	FORCEINLINE float ComputeViewMaxBackgroundCocRadius(float HorizontalResolution) const
 	{
 		return FMath::Min(FMath::Max(InfinityBackgroundCocRadius, MaxDepthBlurRadius), MaxBackgroundCocRadius) * HorizontalResolution;
 	}
 	
-	/** Returns limit(DepthToHalfResCocRadius) for SceneDepth -> 0.
+	/** Returns limit(SceneDepthToCocRadius) for SceneDepth -> 0.
 	 *
 	 * Note: this return negative or null value since this is foreground.
 	 */
diff --git a/Engine/Source/Runtime/Renderer/Private/PostProcess/DiaphragmDOFUtils.cpp b/Engine/Source/Runtime/Renderer/Private/PostProcess/DiaphragmDOFUtils.cpp
--- a/Engine/Source/Runtime/Renderer/Private/PostProcess/DiaphragmDOFUtils.cpp
+++ b/Engine/Source/Runtime/Renderer/Private/PostProcess/DiaphragmDOFUtils.cpp
@@ -102,6 +102,7 @@ void DiaphragmDOF::FPhysicalCocModel::Compile(const FViewInfo& View)
 	// Fetches DOF settings.
 	{
 		FocusDistance = View.FinalPostProcessSettings.DepthOfFieldFocalDistance;
+		Squeeze = FMath::Clamp(View.FinalPostProcessSettings.DepthOfFieldSqueezeFactor, 1.0f, 2.0f);
 
 		// -because foreground Coc are negative.
 		MinForegroundCocRadius = -CVarMaxForegroundRadius.GetValueOnRenderThread();
