Moved cine camera lens and filmback presets from the CineCameraComponent to a new CineCameraSettings class that appears in the project settings. This allows users to view and edit their own presets easily without needing to edit an INI file

diff --git a/Engine/Config/BaseEngine.ini b/Engine/Config/BaseEngine.ini
--- a/Engine/Config/BaseEngine.ini
+++ b/Engine/Config/BaseEngine.ini
@@ -2913,7 +2913,7 @@ GameHUDBind=(Key=Tilde,bShift=False,bCtrl=True,bAlt=False,bCmd=False)
 [/Script/Engine.PlayerCameraManager]
 ServerUpdateCameraTimeout=2.0
 
-[/Script/CinematicCamera.CineCameraComponent]
+[/Script/CinematicCamera.CineCameraSettings]
 +FilmbackPresets=(Name="16:9 Film",FilmbackSettings=(SensorWidth=24.00,SensorHeight=13.5))
 +FilmbackPresets=(Name="16:9 Digital Film",FilmbackSettings=(SensorWidth=23.76,SensorHeight=13.365))
 +FilmbackPresets=(Name="16:9 DSLR",FilmbackSettings=(SensorWidth=36,SensorHeight=20.25))
@@ -2927,7 +2927,6 @@ ServerUpdateCameraTimeout=2.0
 +FilmbackPresets=(Name="APS-C (Canon)",FilmbackSettings=(SensorWidth=22.2,SensorHeight=14.8))
 +FilmbackPresets=(Name="Full Frame DSLR",FilmbackSettings=(SensorWidth=36,SensorHeight=24))
 +FilmbackPresets=(Name="Micro Four Thirds",FilmbackSettings=(SensorWidth=17.3,SensorHeight=13))
-DefaultFilmbackPresetName="16:9 DSLR" ; Deprecated, here for backwards compatibility
 DefaultFilmbackPreset="16:9 Digital Film"
 +LensPresets=(Name="12mm Prime f/2.8",LensSettings=(MinFocalLength=12,MaxFocalLength=12,MinFStop=2.8,MaxFStop=22,DiaphragmBladeCount=7))
 +LensPresets=(Name="30mm Prime f/1.4",LensSettings=(MinFocalLength=30,MaxFocalLength=30,MinFStop=1.4,MaxFStop=22,DiaphragmBladeCount=7))
diff --git a/Engine/Source/Editor/DetailCustomizations/Private/CameraCropSettingsCustomization.cpp b/Engine/Source/Editor/DetailCustomizations/Private/CameraCropSettingsCustomization.cpp
--- a/Engine/Source/Editor/DetailCustomizations/Private/CameraCropSettingsCustomization.cpp
+++ b/Engine/Source/Editor/DetailCustomizations/Private/CameraCropSettingsCustomization.cpp
@@ -6,13 +6,13 @@
 #include "DetailWidgetRow.h"
 #include "DetailLayoutBuilder.h"
 #include "ScopedTransaction.h"
-#include "CineCameraComponent.h"
+#include "CineCameraSettings.h"
 
 #define LOCTEXT_NAMESPACE "CameraCropSettingsCustomization"
 
 FCameraCropSettingsCustomization::FCameraCropSettingsCustomization()
 {
-	TArray<FNamedPlateCropPreset> const& Presets = UCineCameraComponent::GetCropPresets();
+	TArray<FNamedPlateCropPreset> const& Presets = UCineCameraSettings::GetCropPresets();
 
 	int32 const NumPresets = Presets.Num();
 	// first create preset combo list
@@ -34,28 +34,35 @@ TSharedRef<IPropertyTypeCustomization> FCameraCropSettingsCustomization::MakeIns
 
 void FCameraCropSettingsCustomization::CustomizeHeader(TSharedRef<IPropertyHandle> StructPropertyHandle, FDetailWidgetRow& HeaderRow, IPropertyTypeCustomizationUtils& CustomizationUtils)
 {
-	HeaderRow.
-		NameContent()
-		[
-			StructPropertyHandle->CreatePropertyNameWidget()
-		]
-		.ValueContent()
-		.MaxDesiredWidth(0.f)
-		[
-			SAssignNew(PresetComboBox, SComboBox< TSharedPtr<FString> >)
-			.OptionsSource(&PresetComboList)
-			.OnGenerateWidget(this, &FCameraCropSettingsCustomization::MakePresetComboWidget)
-			.OnSelectionChanged(this, &FCameraCropSettingsCustomization::OnPresetChanged)
-			.IsEnabled(FSlateApplication::Get().GetNormalExecutionAttribute())
-			.ContentPadding(2)
-			.Content()
+	// We only want the dropdown list outside of the settings class as the settings class is the thing
+	// defining the presets we use for the dropdown
+	const bool bInSettingsClass = StructPropertyHandle->GetOuterBaseClass() == UCineCameraSettings::StaticClass();
+
+	if (!bInSettingsClass)
+	{
+		HeaderRow.
+			NameContent()
 			[
-				SNew(STextBlock)
-				.Text(this, &FCameraCropSettingsCustomization::GetPresetComboBoxContent)
-				.Font(IDetailLayoutBuilder::GetDetailFont())
-				.ToolTipText(this, &FCameraCropSettingsCustomization::GetPresetComboBoxContent)
+				StructPropertyHandle->CreatePropertyNameWidget()
 			]
-		];
+			.ValueContent()
+			.MaxDesiredWidth(0.f)
+			[
+				SAssignNew(PresetComboBox, SComboBox< TSharedPtr<FString> >)
+				.OptionsSource(&PresetComboList)
+				.OnGenerateWidget(this, &FCameraCropSettingsCustomization::MakePresetComboWidget)
+				.OnSelectionChanged(this, &FCameraCropSettingsCustomization::OnPresetChanged)
+				.IsEnabled(FSlateApplication::Get().GetNormalExecutionAttribute())
+				.ContentPadding(2)
+				.Content()
+				[
+					SNew(STextBlock)
+					.Text(this, &FCameraCropSettingsCustomization::GetPresetComboBoxContent)
+					.Font(IDetailLayoutBuilder::GetDetailFont())
+					.ToolTipText(this, &FCameraCropSettingsCustomization::GetPresetComboBoxContent)
+				]
+			];
+	}
 }
 
 
@@ -98,7 +105,7 @@ void FCameraCropSettingsCustomization::OnPresetChanged(TSharedPtr<FString> NewSe
 		FString const NewPresetName = *NewSelection.Get();
 
 		// search presets for one that matches
-		TArray<FNamedPlateCropPreset> const& Presets = UCineCameraComponent::GetCropPresets();
+		TArray<FNamedPlateCropPreset> const& Presets = UCineCameraSettings::GetCropPresets();
 		int32 const NumPresets = Presets.Num();
 		for (int32 PresetIdx = 0; PresetIdx < NumPresets; ++PresetIdx)
 		{
@@ -132,7 +139,7 @@ TSharedPtr<FString> FCameraCropSettingsCustomization::GetPresetString() const
  	AspectRatioHandle->GetValue(AspectRatio);
  
 	// search presets for one that matches
-	TArray<FNamedPlateCropPreset> const& Presets = UCineCameraComponent::GetCropPresets();
+	TArray<FNamedPlateCropPreset> const& Presets = UCineCameraSettings::GetCropPresets();
 	int32 const NumPresets = Presets.Num();
 	for (int32 PresetIdx = 0; PresetIdx < NumPresets; ++PresetIdx)
 	{
diff --git a/Engine/Source/Editor/DetailCustomizations/Private/CameraFilmbackSettingsCustomization.cpp b/Engine/Source/Editor/DetailCustomizations/Private/CameraFilmbackSettingsCustomization.cpp
--- a/Engine/Source/Editor/DetailCustomizations/Private/CameraFilmbackSettingsCustomization.cpp
+++ b/Engine/Source/Editor/DetailCustomizations/Private/CameraFilmbackSettingsCustomization.cpp
@@ -28,9 +28,9 @@ class SWidget;
 
 #define LOCTEXT_NAMESPACE "CameraFilmbackSettingsCustomization"
 
-FCameraFilmbackSettingsCustomization::FCameraFilmbackSettingsCustomization()
+void FCameraFilmbackSettingsCustomization::BuildPresetComboList()
 {
-	TArray<FNamedFilmbackPreset> const& Presets = UCineCameraComponent::GetFilmbackPresets();
+	TArray<FNamedFilmbackPreset> const& Presets = UCineCameraSettings::GetFilmbackPresets();
 
 	int32 const NumPresets = Presets.Num();
 	// first create preset combo list
@@ -46,6 +46,11 @@ FCameraFilmbackSettingsCustomization::FCameraFilmbackSettingsCustomization()
 	}
 }
 
+FCameraFilmbackSettingsCustomization::FCameraFilmbackSettingsCustomization()
+{
+	BuildPresetComboList();
+}
+
 TSharedRef<IPropertyTypeCustomization> FCameraFilmbackSettingsCustomization::MakeInstance()
 {
 	return MakeShareable(new FCameraFilmbackSettingsCustomization);
@@ -57,24 +62,34 @@ void FCameraFilmbackSettingsCustomization::CustomizeHeader(TSharedRef<IPropertyH
 		NameContent()
 		[
 			StructPropertyHandle->CreatePropertyNameWidget()
-		]
-		.ValueContent()
-		.MaxDesiredWidth(0.f)
-		[
-			SAssignNew(PresetComboBox, SComboBox< TSharedPtr<FString> >)
-			.OptionsSource(&PresetComboList)
-			.OnGenerateWidget(this, &FCameraFilmbackSettingsCustomization::MakePresetComboWidget)
-			.OnSelectionChanged(this, &FCameraFilmbackSettingsCustomization::OnPresetChanged)
-			.IsEnabled(this, &FCameraFilmbackSettingsCustomization::IsPresetEnabled)
-			.ContentPadding(2)
-			.Content()
-			[
-				SNew(STextBlock)
-				.Text(this, &FCameraFilmbackSettingsCustomization::GetPresetComboBoxContent)
-				.Font(IDetailLayoutBuilder::GetDetailFont())
-				.ToolTipText(this, &FCameraFilmbackSettingsCustomization::GetPresetComboBoxContent)
-			]
 		];
+
+	// We only want the dropdown list outside of the settings class as the settings class is the thing
+	// defining the presets we use for the dropdown
+	const bool bInSettingsClass = StructPropertyHandle->GetOuterBaseClass() == UCineCameraSettings::StaticClass();
+
+	if (!bInSettingsClass)
+	{
+		HeaderRow.
+			ValueContent()
+			.MaxDesiredWidth(0.f)
+			[
+				SAssignNew(PresetComboBox, SComboBox< TSharedPtr<FString> >)
+				.OptionsSource(&PresetComboList)
+				.OnGenerateWidget(this, &FCameraFilmbackSettingsCustomization::MakePresetComboWidget)
+				.OnSelectionChanged(this, &FCameraFilmbackSettingsCustomization::OnPresetChanged)
+				.OnComboBoxOpening(this, &FCameraFilmbackSettingsCustomization::BuildPresetComboList)
+				.IsEnabled(this, &FCameraFilmbackSettingsCustomization::IsPresetEnabled)
+				.ContentPadding(2)
+				.Content()
+				[
+					SNew(STextBlock)
+					.Text(this, &FCameraFilmbackSettingsCustomization::GetPresetComboBoxContent)
+					.Font(IDetailLayoutBuilder::GetDetailFont())
+					.ToolTipText(this, &FCameraFilmbackSettingsCustomization::GetPresetComboBoxContent)
+				]
+			];
+	}
 }
 
 void FCameraFilmbackSettingsCustomization::CustomizeChildren(TSharedRef<IPropertyHandle> StructPropertyHandle, class IDetailChildrenBuilder& ChildBuilder, IPropertyTypeCustomizationUtils& StructCustomizationUtils)
@@ -130,7 +145,7 @@ void FCameraFilmbackSettingsCustomization::OnPresetChanged(TSharedPtr<FString> N
 		FString const NewPresetName = *NewSelection.Get();
 
 		// search presets for one that matches
-		TArray<FNamedFilmbackPreset> const& Presets = UCineCameraComponent::GetFilmbackPresets();
+		TArray<FNamedFilmbackPreset> const& Presets = UCineCameraSettings::GetFilmbackPresets();
 		int32 const NumPresets = Presets.Num();
 		for (int32 PresetIdx = 0; PresetIdx < NumPresets; ++PresetIdx)
 		{
@@ -176,7 +191,7 @@ TSharedPtr<FString> FCameraFilmbackSettingsCustomization::GetPresetString() cons
 	SensorHeightHandle->GetValue(CurSensorHeight);
 
 	// search presets for one that matches
-	TArray<FNamedFilmbackPreset> const& Presets = UCineCameraComponent::GetFilmbackPresets();
+	TArray<FNamedFilmbackPreset> const& Presets = UCineCameraSettings::GetFilmbackPresets();
 	int32 const NumPresets = Presets.Num();
 	for (int32 PresetIdx = 0; PresetIdx < NumPresets; ++PresetIdx)
 	{
diff --git a/Engine/Source/Editor/DetailCustomizations/Private/CameraFilmbackSettingsCustomization.h b/Engine/Source/Editor/DetailCustomizations/Private/CameraFilmbackSettingsCustomization.h
--- a/Engine/Source/Editor/DetailCustomizations/Private/CameraFilmbackSettingsCustomization.h
+++ b/Engine/Source/Editor/DetailCustomizations/Private/CameraFilmbackSettingsCustomization.h
@@ -42,6 +42,8 @@ protected:
 	TSharedPtr<class SComboBox< TSharedPtr<FString> > > PresetComboBox;
 	TArray< TSharedPtr< FString > >						PresetComboList;
 
+	void BuildPresetComboList();
+
 	bool IsPresetEnabled() const;
 	void OnPresetChanged(TSharedPtr<FString> NewSelection, ESelectInfo::Type SelectInfo);
 	TSharedRef<SWidget> MakePresetComboWidget(TSharedPtr<FString> InItem);
diff --git a/Engine/Source/Editor/DetailCustomizations/Private/CameraLensSettingsCustomization.cpp b/Engine/Source/Editor/DetailCustomizations/Private/CameraLensSettingsCustomization.cpp
--- a/Engine/Source/Editor/DetailCustomizations/Private/CameraLensSettingsCustomization.cpp
+++ b/Engine/Source/Editor/DetailCustomizations/Private/CameraLensSettingsCustomization.cpp
@@ -28,9 +28,9 @@ class SWidget;
 
 #define LOCTEXT_NAMESPACE "CameraLensSettingsCustomization"
 
-FCameraLensSettingsCustomization::FCameraLensSettingsCustomization()
+void FCameraLensSettingsCustomization::BuildPresetComboList()
 {
-	TArray<FNamedLensPreset> const& Presets = UCineCameraComponent::GetLensPresets();
+	TArray<FNamedLensPreset> const& Presets = UCineCameraSettings::GetLensPresets();
 
 	int32 const NumPresets = Presets.Num();
 	// first create preset combo list
@@ -46,6 +46,11 @@ FCameraLensSettingsCustomization::FCameraLensSettingsCustomization()
 	}
 }
 
+FCameraLensSettingsCustomization::FCameraLensSettingsCustomization()
+{
+	BuildPresetComboList();
+}
+
 TSharedRef<IPropertyTypeCustomization> FCameraLensSettingsCustomization::MakeInstance()
 {
 	return MakeShareable(new FCameraLensSettingsCustomization);
@@ -57,24 +62,34 @@ void FCameraLensSettingsCustomization::CustomizeHeader(TSharedRef<IPropertyHandl
 		NameContent()
 		[
 			StructPropertyHandle->CreatePropertyNameWidget()
-		]
-		.ValueContent()
-		.MaxDesiredWidth(0.f)
-		[
-			SAssignNew(PresetComboBox, SComboBox< TSharedPtr<FString> >)
-			.OptionsSource(&PresetComboList)
-			.OnGenerateWidget(this, &FCameraLensSettingsCustomization::MakePresetComboWidget)
-			.OnSelectionChanged(this, &FCameraLensSettingsCustomization::OnPresetChanged)
-			.IsEnabled(FSlateApplication::Get().GetNormalExecutionAttribute())
-			.ContentPadding(2)
-			.Content()
-			[
-				SNew(STextBlock)
-				.Text(this, &FCameraLensSettingsCustomization::GetPresetComboBoxContent)
-				.Font(IDetailLayoutBuilder::GetDetailFont())
-				.ToolTipText(this, &FCameraLensSettingsCustomization::GetPresetComboBoxContent)
-			]
 		];
+
+	// We only want the dropdown list outside of the settings class as the settings class is the thing
+	// defining the presets we use for the dropdown
+	const bool bInSettingsClass = StructPropertyHandle->GetOuterBaseClass() == UCineCameraSettings::StaticClass();
+
+	if (!bInSettingsClass)
+	{
+		HeaderRow.
+			ValueContent()
+			.MaxDesiredWidth(0.f)
+			[
+				SAssignNew(PresetComboBox, SComboBox< TSharedPtr<FString> >)
+				.OptionsSource(&PresetComboList)
+				.OnGenerateWidget(this, &FCameraLensSettingsCustomization::MakePresetComboWidget)
+				.OnSelectionChanged(this, &FCameraLensSettingsCustomization::OnPresetChanged)
+				.OnComboBoxOpening(this, &FCameraLensSettingsCustomization::BuildPresetComboList)
+				.IsEnabled(FSlateApplication::Get().GetNormalExecutionAttribute())
+				.ContentPadding(2)
+				.Content()
+				[
+					SNew(STextBlock)
+					.Text(this, &FCameraLensSettingsCustomization::GetPresetComboBoxContent)
+					.Font(IDetailLayoutBuilder::GetDetailFont())
+					.ToolTipText(this, &FCameraLensSettingsCustomization::GetPresetComboBoxContent)
+				]
+			];
+	}
 }
 
 
@@ -129,7 +144,7 @@ void FCameraLensSettingsCustomization::OnPresetChanged(TSharedPtr<FString> NewSe
 		FString const NewPresetName = *NewSelection.Get();
 
 		// search presets for one that matches
-		TArray<FNamedLensPreset> const& Presets = UCineCameraComponent::GetLensPresets();
+		TArray<FNamedLensPreset> const& Presets = UCineCameraSettings::GetLensPresets();
 		int32 const NumPresets = Presets.Num();
 		for (int32 PresetIdx = 0; PresetIdx < NumPresets; ++PresetIdx)
 		{
@@ -195,7 +210,7 @@ TSharedPtr<FString> FCameraLensSettingsCustomization::GetPresetString() const
 	DiaphragmBladeCountHandle->GetValue(DiaphragmBladeCount);
 
 	// search presets for one that matches
-	TArray<FNamedLensPreset> const& Presets = UCineCameraComponent::GetLensPresets();
+	TArray<FNamedLensPreset> const& Presets = UCineCameraSettings::GetLensPresets();
 	int32 const NumPresets = Presets.Num();
 	for (int32 PresetIdx = 0; PresetIdx < NumPresets; ++PresetIdx)
 	{
diff --git a/Engine/Source/Editor/DetailCustomizations/Private/CameraLensSettingsCustomization.h b/Engine/Source/Editor/DetailCustomizations/Private/CameraLensSettingsCustomization.h
--- a/Engine/Source/Editor/DetailCustomizations/Private/CameraLensSettingsCustomization.h
+++ b/Engine/Source/Editor/DetailCustomizations/Private/CameraLensSettingsCustomization.h
@@ -29,7 +29,7 @@ public:
 	FCameraLensSettingsCustomization();
 
 	static TSharedRef<IPropertyTypeCustomization> MakeInstance();
-
+	
 	/** IPropertyTypeCustomization instance */
 	virtual void CustomizeHeader(TSharedRef<IPropertyHandle> PropertyHandle, FDetailWidgetRow& HeaderRow, IPropertyTypeCustomizationUtils& CustomizationUtils) override;
 	virtual void CustomizeChildren( TSharedRef<IPropertyHandle> StructPropertyHandle, class IDetailChildrenBuilder& ChildBuilder, IPropertyTypeCustomizationUtils& StructCustomizationUtils ) override;
@@ -47,6 +47,8 @@ protected:
 	TSharedPtr<class SComboBox< TSharedPtr<FString> > > PresetComboBox;
 	TArray< TSharedPtr< FString > >						PresetComboList;
 
+	void BuildPresetComboList();
+	
 	void OnPresetChanged(TSharedPtr<FString> NewSelection, ESelectInfo::Type SelectInfo);
 	TSharedRef<SWidget> MakePresetComboWidget(TSharedPtr<FString> InItem);
 
diff --git a/Engine/Source/Runtime/CinematicCamera/CinematicCamera.Build.cs b/Engine/Source/Runtime/CinematicCamera/CinematicCamera.Build.cs
--- a/Engine/Source/Runtime/CinematicCamera/CinematicCamera.Build.cs
+++ b/Engine/Source/Runtime/CinematicCamera/CinematicCamera.Build.cs
@@ -10,7 +10,10 @@ public class CinematicCamera : ModuleRules
 			new string[] {
 				"Core",
 				"CoreUObject",
+				"DeveloperSettings",
                 "Engine",
+                "Slate",
+                "SlateCore"
 			}
 		);
 
diff --git a/Engine/Source/Runtime/CinematicCamera/Private/CineCameraComponent.cpp b/Engine/Source/Runtime/CinematicCamera/Private/CineCameraComponent.cpp
--- a/Engine/Source/Runtime/CinematicCamera/Private/CineCameraComponent.cpp
+++ b/Engine/Source/Runtime/CinematicCamera/Private/CineCameraComponent.cpp
@@ -37,21 +37,17 @@ UCineCameraComponent::UCineCameraComponent()
 
 	bConstrainAspectRatio = true;
 
-	// Certain default values are set by a config, so Use the archetype to set them in the constructor, so they can be overridden in the editor.
-	UCineCameraComponent* Template = Cast<UCineCameraComponent>(GetArchetype());
 	
-	// Null check here because CDO has no Archetype
-	if (Template)
-	{
-		// default filmback, lens and crop
-		SetFilmbackPresetByNameInternal(Template->DefaultFilmbackPreset, Filmback);
-		SetFilmbackPresetByNameInternal(Template->DefaultFilmbackPresetName_DEPRECATED, FilmbackSettings_DEPRECATED);
-		SetLensPresetByNameInternal(Template->DefaultLensPresetName);
-		SetCropPresetByNameInternal(Template->DefaultCropPresetName);
-		// other lens defaults
-		CurrentAperture = Template->DefaultLensFStop;
-		CurrentFocalLength = Template->DefaultLensFocalLength;
-	}
+	const UCineCameraSettings* CineCameraSettings = GetDefault<UCineCameraSettings>();
+	check(CineCameraSettings);
+	// default filmback
+	SetFilmbackPresetByNameInternal(CineCameraSettings->DefaultFilmbackPreset, Filmback);
+	SetFilmbackPresetByNameInternal(CineCameraSettings->DefaultFilmbackPreset, FilmbackSettings_DEPRECATED);
+	SetLensPresetByNameInternal(CineCameraSettings->DefaultLensPresetName);
+	// other lens defaults
+	CurrentAperture = CineCameraSettings->DefaultLensFStop;
+	CurrentFocalLength = CineCameraSettings->DefaultLensFocalLength;
+
 
 	RecalcDerivedData();
 
@@ -76,6 +72,8 @@ UCineCameraComponent::UCineCameraComponent()
 
 void UCineCameraComponent::Serialize(FArchive& Ar)
 {
+	PRAGMA_DISABLE_DEPRECATION_WARNINGS
+	
 	Ar.UsingCustomVersion(FCineCameraObjectVersion::GUID);
 	Ar.UsingCustomVersion(FReleaseObjectVersion::GUID);
 
@@ -89,7 +87,7 @@ void UCineCameraComponent::Serialize(FArchive& Ar)
 			UCineCameraComponent* Template = Cast<UCineCameraComponent>(GetArchetype());
 			if (Template)
 			{
-				TArray<FNamedFilmbackPreset> const& Presets = UCineCameraComponent::GetFilmbackPresets();
+				TArray<FNamedFilmbackPreset> const& Presets = UCineCameraSettings::GetFilmbackPresets();
 				int32 const NumPresets = Presets.Num();
 				for (int32 PresetIdx = 0; PresetIdx < NumPresets; ++PresetIdx)
 				{
@@ -116,6 +114,8 @@ void UCineCameraComponent::Serialize(FArchive& Ar)
 			Filmback = FilmbackSettings_DEPRECATED;
 		}
 	}
+	
+	PRAGMA_ENABLE_DEPRECATION_WARNINGS
 }
 
 void UCineCameraComponent::PostInitProperties()
@@ -311,7 +311,7 @@ float UCineCameraComponent::GetVerticalFieldOfView() const
 
 FString UCineCameraComponent::GetFilmbackPresetName() const
 {
-	TArray<FNamedFilmbackPreset> const& Presets = UCineCameraComponent::GetFilmbackPresets();
+	TArray<FNamedFilmbackPreset> const& Presets = UCineCameraSettings::GetFilmbackPresets();
 	int32 const NumPresets = Presets.Num();
 	for (int32 PresetIdx = 0; PresetIdx < NumPresets; ++PresetIdx)
 	{
@@ -325,6 +325,16 @@ FString UCineCameraComponent::GetFilmbackPresetName() const
 	return FString();
 }
 
+FString UCineCameraComponent::GetDefaultFilmbackPresetName() const
+{
+	// Try to get the default from the CineCameraSettings object but fallback to an empty string on failure
+	if (const UCineCameraSettings* CineCameraSettings = GetDefault<UCineCameraSettings>())
+	{
+		return CineCameraSettings->DefaultFilmbackPreset;
+	}
+	return FString();
+}
+
 void UCineCameraComponent::SetFilmbackPresetByName(const FString& InPresetName)
 {
 	SetFilmbackPresetByNameInternal(InPresetName, Filmback);
@@ -334,7 +344,7 @@ void UCineCameraComponent::SetFilmbackPresetByName(const FString& InPresetName)
 
 void UCineCameraComponent::SetFilmbackPresetByNameInternal(const FString& InPresetName, FCameraFilmbackSettings& InOutFilmbackSettings)
 {
-	TArray<FNamedFilmbackPreset> const& Presets = UCineCameraComponent::GetFilmbackPresets();
+	TArray<FNamedFilmbackPreset> const& Presets = UCineCameraSettings::GetFilmbackPresets();
 	int32 const NumPresets = Presets.Num();
 	for (int32 PresetIdx = 0; PresetIdx < NumPresets; ++PresetIdx)
 	{
@@ -349,7 +359,7 @@ void UCineCameraComponent::SetFilmbackPresetByNameInternal(const FString& InPres
 
 FString UCineCameraComponent::GetLensPresetName() const
 {
-	TArray<FNamedLensPreset> const& Presets = UCineCameraComponent::GetLensPresets();
+	TArray<FNamedLensPreset> const& Presets = UCineCameraSettings::GetLensPresets();
 	int32 const NumPresets = Presets.Num();
 	for (int32 PresetIdx = 0; PresetIdx < NumPresets; ++PresetIdx)
 	{
@@ -372,7 +382,7 @@ void UCineCameraComponent::SetLensPresetByName(const FString& InPresetName)
 
 void UCineCameraComponent::SetLensPresetByNameInternal(const FString& InPresetName)
 {
-	TArray<FNamedLensPreset> const& Presets = UCineCameraComponent::GetLensPresets();
+	TArray<FNamedLensPreset> const& Presets = UCineCameraSettings::GetLensPresets();
 	int32 const NumPresets = Presets.Num();
 	for (int32 PresetIdx = 0; PresetIdx < NumPresets; ++PresetIdx)
 	{
@@ -387,7 +397,7 @@ void UCineCameraComponent::SetLensPresetByNameInternal(const FString& InPresetNa
 
 FString UCineCameraComponent::GetCropPresetName() const
 {
-	TArray<FNamedPlateCropPreset> const& Presets = UCineCameraComponent::GetCropPresets();
+	TArray<FNamedPlateCropPreset> const& Presets = UCineCameraSettings::GetCropPresets();
 	int32 const NumPresets = Presets.Num();
 	for (int32 PresetIdx = 0; PresetIdx < NumPresets; ++PresetIdx)
 	{
@@ -410,7 +420,7 @@ void UCineCameraComponent::SetCropPresetByName(const FString& InPresetName)
 
 void UCineCameraComponent::SetCropPresetByNameInternal(const FString& InPresetName)
 {
-	TArray<FNamedPlateCropPreset> const& Presets = UCineCameraComponent::GetCropPresets();
+	TArray<FNamedPlateCropPreset> const& Presets = UCineCameraSettings::GetCropPresets();
 	int32 const NumPresets = Presets.Num();
 	for (int32 PresetIdx = 0; PresetIdx < NumPresets; ++PresetIdx)
 	{
@@ -433,31 +443,25 @@ float UCineCameraComponent::GetWorldToMetersScale() const
 // static
 TArray<FNamedFilmbackPreset> UCineCameraComponent::GetFilmbackPresetsCopy()
 {
-	return GetDefault<UCineCameraComponent>()->FilmbackPresets;
+	return UCineCameraSettings::GetFilmbackPresets();
 }
 
 // static
 TArray<FNamedLensPreset> UCineCameraComponent::GetLensPresetsCopy()
 {
-	return GetDefault<UCineCameraComponent>()->LensPresets;
+	return UCineCameraSettings::GetLensPresets();
 }
 
 // static
 TArray<FNamedFilmbackPreset> const& UCineCameraComponent::GetFilmbackPresets()
 {
-	return GetDefault<UCineCameraComponent>()->FilmbackPresets;
+	return UCineCameraSettings::GetFilmbackPresets();
 }
 
 // static
 TArray<FNamedLensPreset> const& UCineCameraComponent::GetLensPresets()
 {
-	return GetDefault<UCineCameraComponent>()->LensPresets;
-}
-
-// static
-TArray<FNamedPlateCropPreset> const& UCineCameraComponent::GetCropPresets()
-{
-	return GetDefault<UCineCameraComponent>()->CropPresets;
+	return UCineCameraSettings::GetLensPresets();
 }
 
 void UCineCameraComponent::RecalcDerivedData()
@@ -555,7 +559,7 @@ FText UCineCameraComponent::GetFilmbackText() const
 	const float SensorHeight = Filmback.SensorHeight;
 
 	// Search presets for one that matches
-	const FNamedFilmbackPreset* Preset = UCineCameraComponent::GetFilmbackPresets().FindByPredicate([&](const FNamedFilmbackPreset& InPreset) {
+	const FNamedFilmbackPreset* Preset = UCineCameraSettings::GetFilmbackPresets().FindByPredicate([&](const FNamedFilmbackPreset& InPreset) {
 		return InPreset.FilmbackSettings.SensorWidth == SensorWidth && InPreset.FilmbackSettings.SensorHeight == SensorHeight;
 	});
 
diff --git a/Engine/Source/Runtime/CinematicCamera/Private/CineCameraSettings.cpp b/Engine/Source/Runtime/CinematicCamera/Private/CineCameraSettings.cpp
new file mode 100644
--- /dev/null
+++ b/Engine/Source/Runtime/CinematicCamera/Private/CineCameraSettings.cpp
@@ -0,0 +1,241 @@
+// Copyright Epic Games, Inc. All Rights Reserved.
+
+#include "CineCameraSettings.h"
+
+#include "CineCameraComponent.h"
+#include "Misc/ConfigCacheIni.h"
+
+#include "Framework/Notifications/NotificationManager.h"
+#include "Widgets/Notifications/SNotificationList.h"
+
+#define LOCTEXT_NAMESPACE "CineCameraSettings"
+
+const FString UCineCameraSettings::CineCameraConfigSection = TEXT("/Script/CinematicCamera.CineCameraComponent");
+
+void UCineCameraSettings::PostInitProperties()
+{
+	Super::PostInitProperties();
+	
+	if (GConfig && GConfig->DoesSectionExist(*CineCameraConfigSection, GEngineIni))
+	{
+		FNotificationInfo NotificationInfo(LOCTEXT("UpdateSettingsNotification", "CineCamera Settings were found in an old config location. Do you want to attempt to automatically merge them to the new location?"));
+		NotificationInfo.FadeOutDuration = 0.5f;
+		NotificationInfo.bFireAndForget = false;
+		NotificationInfo.bUseSuccessFailIcons = true;
+
+		const FNotificationButtonInfo OkButton(LOCTEXT("MergeButtonText", "Merge"), FText::GetEmpty(), FSimpleDelegate::CreateUObject(this, &UCineCameraSettings::CopyOldConfigSettings), SNotificationItem::ECompletionState::CS_None);
+		const FNotificationButtonInfo CancelButton(LOCTEXT("CancelButtonText", "Cancel"), FText::GetEmpty(), FSimpleDelegate::CreateUObject(this, &UCineCameraSettings::CloseNotification), SNotificationItem::ECompletionState::CS_None);
+
+		NotificationInfo.ButtonDetails.Add(OkButton);
+		NotificationInfo.ButtonDetails.Add(CancelButton);
+
+		Notification = FSlateNotificationManager::Get().AddNotification(NotificationInfo);
+	}
+}
+
+void UCineCameraSettings::SetDefaultLensPresetName(const FString InDefaultLensPresetName)
+{
+	DefaultLensPresetName = InDefaultLensPresetName;
+	SaveConfig();
+}
+
+void UCineCameraSettings::SetDefaultLensFocalLength(const float InDefaultLensFocalLength)
+{
+	DefaultLensFocalLength = InDefaultLensFocalLength;
+	SaveConfig();
+}
+
+void UCineCameraSettings::SetDefaultLensFStop(const float InDefaultLensFStop)
+{
+	DefaultLensFStop = InDefaultLensFStop;
+	SaveConfig();
+}
+
+void UCineCameraSettings::SetLensPresets(const TArray<FNamedLensPreset>& InLensPresets)
+{
+	LensPresets = InLensPresets;
+	SaveConfig();
+}
+
+void UCineCameraSettings::SetDefaultFilmbackPreset(const FString InDefaultFilmbackPreset)
+{
+	DefaultFilmbackPreset = InDefaultFilmbackPreset;
+	SaveConfig();
+}
+
+void UCineCameraSettings::SetFilmbackPresets(const TArray<FNamedFilmbackPreset>& InFilmbackPresets)
+{
+	FilmbackPresets = InFilmbackPresets;
+	SaveConfig();
+}
+
+void UCineCameraSettings::SetDefaultCropPresetName(const FString InDefaultCropPresetName)
+{
+	DefaultCropPresetName = InDefaultCropPresetName;
+	SaveConfig();
+}
+
+void UCineCameraSettings::SetCropPresets(const TArray<FNamedPlateCropPreset>& InCropPresets)
+{
+	CropPresets = InCropPresets;
+	SaveConfig();
+}
+
+TArray<FNamedLensPreset> const& UCineCameraSettings::GetLensPresets()
+{
+	return GetDefault<UCineCameraSettings>()->LensPresets;
+}
+
+TArray<FNamedFilmbackPreset> const& UCineCameraSettings::GetFilmbackPresets()
+{
+	return GetDefault<UCineCameraSettings>()->FilmbackPresets;
+}
+
+TArray<FNamedPlateCropPreset> const& UCineCameraSettings::GetCropPresets()
+{
+	return GetDefault<UCineCameraSettings>()->CropPresets;
+}
+
+UCineCameraSettings* UCineCameraSettings::GetCineCameraSettings()
+{
+	return GetMutableDefault<UCineCameraSettings>();
+}
+
+TArray<FString> UCineCameraSettings::GetLensPresetNames() const
+{
+	TArray<FString> LensPresetNames;
+	for (const FNamedLensPreset& LensPreset : LensPresets)
+	{
+		LensPresetNames.Emplace(LensPreset.Name);
+	}
+
+	return LensPresetNames;
+}
+
+TArray<FString> UCineCameraSettings::GetFilmbackPresetNames() const
+{	
+	TArray<FString> FilmbackPresetNames;
+	for (const FNamedFilmbackPreset& FilmbackPreset : FilmbackPresets)
+	{
+		FilmbackPresetNames.Emplace(FilmbackPreset.Name);
+	}
+	
+	return FilmbackPresetNames;
+}
+
+TArray<FString> UCineCameraSettings::GetCropPresetNames() const
+{
+	TArray<FString> CropPresetNames;
+	for (const FNamedPlateCropPreset& CropPreset : CropPresets)
+	{
+		CropPresetNames.Emplace(CropPreset.Name);
+	}
+
+	return CropPresetNames;
+}
+
+void UCineCameraSettings::CloseNotification()
+{
+	if (Notification)
+	{
+		Notification->ExpireAndFadeout();
+		Notification = nullptr;
+	}
+	
+	FNotificationInfo NotificationInfo(LOCTEXT("SettingsMergeCancelled", "Please manually fix up and remove the values from the CineCameraComponent section"));
+	NotificationInfo.ExpireDuration = 8.0f;
+	NotificationInfo.bFireAndForget = true;
+
+	FSlateNotificationManager::Get().AddNotification(NotificationInfo);
+}
+
+void UCineCameraSettings::CopyOldConfigSettings()
+{	
+	const FString SettingsConfigSection(TEXT("/Script/CinematicCamera.CineCameraSettings"));
+
+	FString OldDefaultLensPresetName;
+	if (GConfig->GetString(*CineCameraConfigSection, TEXT("DefaultLensPresetName"), OldDefaultLensPresetName, GEngineIni))
+	{
+		GConfig->SetString(*SettingsConfigSection, TEXT("DefaultLensPresetName"), *OldDefaultLensPresetName, GEngineIni);
+	}
+
+	float OldDefaultLensFocalLength;
+	if (GConfig->GetFloat(*CineCameraConfigSection, TEXT("DefaultLensFocalLength"), OldDefaultLensFocalLength, GEngineIni))
+	{
+		GConfig->SetFloat(*SettingsConfigSection, TEXT("DefaultLensFocalLength"), OldDefaultLensFocalLength, GEngineIni);
+	}
+
+	float OldDefaultLensFStop;
+	if (GConfig->GetFloat(*CineCameraConfigSection, TEXT("DefaultLensFStop"), OldDefaultLensFStop, GEngineIni))
+	{
+		GConfig->SetFloat(*SettingsConfigSection, TEXT("DefaultLensFStop"), OldDefaultLensFStop, GEngineIni);
+	}
+
+	TArray<FString> OldLensPresets;
+	if (GConfig->GetArray(*CineCameraConfigSection, TEXT("LensPresets"), OldLensPresets, GEngineIni))
+	{
+		TArray<FString> CurrentLensPresets;
+		GConfig->GetArray(*SettingsConfigSection, TEXT("LensPresets"), CurrentLensPresets, GEngineIni);
+		for (const FString& OldLensPreset : OldLensPresets)
+		{
+			// If a preset already exists with that name then update it with the old value
+			// otherwise add the preset to the list
+			FString PresetName, PresetValue;
+			OldLensPreset.Split(",", &PresetName, &PresetValue);
+			if (FString* MatchingLensPreset = CurrentLensPresets.FindByPredicate([PresetName](const FString& CurrentPreset){ return CurrentPreset.StartsWith(PresetName); }))
+			{
+				*MatchingLensPreset = OldLensPreset;	
+			}
+			else
+			{
+				CurrentLensPresets.Add(OldLensPreset);
+			}
+		}
+		GConfig->SetArray(*SettingsConfigSection, TEXT("LensPresets"), CurrentLensPresets, GEngineIni);
+	}
+
+	FString OldDefaultFilmbackPreset;
+	if (GConfig->GetString(*CineCameraConfigSection, TEXT("DefaultFilmbackPreset"), OldDefaultFilmbackPreset, GEngineIni))
+	{
+		GConfig->SetString(*SettingsConfigSection, TEXT("DefaultFilmbackPreset"), *OldDefaultFilmbackPreset, GEngineIni);
+	}
+
+	TArray<FString> OldFilmbackPresets;
+	if (GConfig->GetArray(*CineCameraConfigSection, TEXT("FilmbackPresets"), OldFilmbackPresets, GEngineIni))
+	{
+		TArray<FString> CurrentFilmbackPresets;
+		GConfig->GetArray(*SettingsConfigSection, TEXT("FilmbackPresets"), CurrentFilmbackPresets, GEngineIni);
+		for (const FString& OldFilmbackPreset : OldFilmbackPresets)
+		{
+			// If a preset already exists with that name then update it with the old value
+			// otherwise add the preset to the list
+			FString PresetName, PresetValue;
+			OldFilmbackPreset.Split(",", &PresetName, &PresetValue);
+			if (FString* MatchingPreset = CurrentFilmbackPresets.FindByPredicate([PresetName](const FString& CurrentPreset){ return CurrentPreset.StartsWith(PresetName); }))
+			{
+				*MatchingPreset = OldFilmbackPreset;
+			}
+			else
+			{
+				CurrentFilmbackPresets.Add(OldFilmbackPreset);
+			}
+		}
+		GConfig->SetArray(*SettingsConfigSection, TEXT("FilmbackPresets"), CurrentFilmbackPresets, GEngineIni);
+	}
+
+	LoadConfig();
+	
+	if (Notification)
+	{
+		Notification->ExpireAndFadeout();
+		Notification = nullptr;
+	}
+
+	FNotificationInfo NotificationInfo(LOCTEXT("NotifySettingsUpdated", "CineCamera Settings have been successfully merged.\n\nPlease remove the old config values from the CineCameraComponent section"));
+	NotificationInfo.ExpireDuration = 8.0f;
+	NotificationInfo.bFireAndForget = true;
+
+	FSlateNotificationManager::Get().AddNotification(NotificationInfo);
+}
+
+#undef LOCTEXT_NAMESPACE
diff --git a/Engine/Source/Runtime/CinematicCamera/Public/CineCameraComponent.h b/Engine/Source/Runtime/CinematicCamera/Public/CineCameraComponent.h
--- a/Engine/Source/Runtime/CinematicCamera/Public/CineCameraComponent.h
+++ b/Engine/Source/Runtime/CinematicCamera/Public/CineCameraComponent.h
@@ -5,6 +5,7 @@
 #include "CoreMinimal.h"
 #include "UObject/ObjectMacros.h"
 #include "Camera/CameraComponent.h"
+#include "CineCameraSettings.h"
 #include "CineCameraComponent.generated.h"
 
 class AActor;
@@ -13,258 +14,6 @@ class UMaterialInstanceDynamic;
 class UStaticMesh;
 class UStaticMeshComponent;
 
-/** #note, this struct has a details customization in CameraFilmbackSettingsCustomization.cpp/h */
-USTRUCT(BlueprintType)
-struct FCameraFilmbackSettings
-{
-	GENERATED_USTRUCT_BODY()
-
-	/** Horizontal size of filmback or digital sensor, in mm. */
-	UPROPERTY(Interp, EditAnywhere, BlueprintReadWrite, Category = "Filmback", meta = (ClampMin = "0.001", ForceUnits = mm))
-	float SensorWidth;
-
-	/** Vertical size of filmback or digital sensor, in mm. */
-	UPROPERTY(Interp, EditAnywhere, BlueprintReadWrite, Category = "Filmback", meta = (ClampMin = "0.001", ForceUnits = mm))
-	float SensorHeight;
-
-	/** Read-only. Computed from Sensor dimensions. */
-	UPROPERTY(Interp, VisibleAnywhere, BlueprintReadOnly, Category = "Filmback")
-	float SensorAspectRatio;
-
-	bool operator==(const FCameraFilmbackSettings& Other) const
-	{
-		return (SensorWidth == Other.SensorWidth)
-			&& (SensorHeight == Other.SensorHeight);
-	}
-
-	bool operator!=(const FCameraFilmbackSettings& Other) const
-	{
-		return !operator==(Other);
-	}
-
-	FCameraFilmbackSettings()
-		: SensorWidth(24.89f)
-		, SensorHeight(18.67f)
-		, SensorAspectRatio(1.33f)
-	{
-	}
-};
-
-/** A named bundle of filmback settings used to implement filmback presets */
-USTRUCT(BlueprintType)
-struct FNamedFilmbackPreset
-{
-	GENERATED_USTRUCT_BODY()
-
-	/** Name for the preset. */
-	UPROPERTY(BlueprintReadWrite, Category = "Filmback")
-	FString Name;
-
-	UPROPERTY(BlueprintReadWrite, Category = "Filmback")
-	FCameraFilmbackSettings FilmbackSettings;
-};
-
-/** 
- * #note, this struct has a details customization in CameraLensSettingsCustomization.cpp/h
- */
-USTRUCT(BlueprintType)
-struct FCameraLensSettings
-{
-	GENERATED_USTRUCT_BODY()
-
-	/** Default constructor, initializing with default values */
-	FCameraLensSettings()
-		: MinFocalLength(50.f)
-		, MaxFocalLength(50.f)
-		, MinFStop(2.f)
-		, MaxFStop(2.f)
-		, MinimumFocusDistance(15.f)
-		, DiaphragmBladeCount(FPostProcessSettings::kDefaultDepthOfFieldBladeCount)
-	{
-	}
-
-	/** Minimum focal length for this lens */
-	UPROPERTY(EditAnywhere, BlueprintReadWrite, Category = "Lens", meta = (ForceUnits = mm, ClampMin = "0.001"))
-	float MinFocalLength = 0.f;
-
-	/** Maximum focal length for this lens */
-	UPROPERTY(EditAnywhere, BlueprintReadWrite, Category = "Lens", meta = (ForceUnits = mm, ClampMin = "0.001"))
-	float MaxFocalLength = 0.f;
-
-	/** Minimum aperture for this lens (e.g. 2.8 for an f/2.8 lens) */
-	UPROPERTY(EditAnywhere, BlueprintReadWrite, Category = "Lens")
-	float MinFStop = 0.f;
-
-	/** Maximum aperture for this lens (e.g. 2.8 for an f/2.8 lens) */
-	UPROPERTY(EditAnywhere, BlueprintReadWrite, Category = "Lens")
-	float MaxFStop = 0.f;
-
-	/** Shortest distance this lens can focus on. */
-	UPROPERTY(EditAnywhere, BlueprintReadWrite, Category = "Lens", meta = (ForceUnits = mm))
-	float MinimumFocusDistance = 0.f;
-
-	/** Squeeze factor for anamorphic lenses. */
-	UPROPERTY(EditAnywhere, BlueprintReadWrite, Category = "Lens", meta = (ClampMin = "1", ClampMax = "2"))
-	float SqueezeFactor = 1.f;
-
-	/** Number of blades of diaphragm. */
-	UPROPERTY(EditAnywhere, BlueprintReadWrite, Category = "Lens", meta = (ClampMin = "4", ClampMax = "16"))
-	int32 DiaphragmBladeCount = 0;
-
-	bool operator==(const FCameraLensSettings& Other) const
-	{
-		return (MinFocalLength == Other.MinFocalLength)
-			&& (MaxFocalLength == Other.MaxFocalLength)
-			&& (MinFStop == Other.MinFStop)
-			&& (MaxFStop == Other.MaxFStop)
-			&& (MinimumFocusDistance == Other.MinimumFocusDistance)
-			&& (SqueezeFactor == Other.SqueezeFactor)
-			&& (DiaphragmBladeCount == Other.DiaphragmBladeCount);
-	}
-};
-
-/** A named bundle of lens settings used to implement lens presets. */
-USTRUCT(BlueprintType)
-struct FNamedLensPreset
-{
-	GENERATED_USTRUCT_BODY()
-
-	/** Name for the preset. */
-	UPROPERTY(BlueprintReadWrite, Category = "Lens")
-	FString Name;
-
-	UPROPERTY(BlueprintReadWrite, Category = "Lens")
-	FCameraLensSettings LensSettings;
-};
-
-/** 
- * #note, this struct has a details customization in CameraCropSettingsCustomization.cpp/h
- */
-USTRUCT(BlueprintType)
-struct FPlateCropSettings
-{
-	GENERATED_USTRUCT_BODY()
-
-	UPROPERTY(EditAnywhere, BlueprintReadWrite, Category = "Crop", meta=(DisplayName = "Cropped Aspect Ratio"))
-	float AspectRatio = 0.f;
-
-	bool operator==(const FPlateCropSettings& Other) const
-	{
-		return (AspectRatio == Other.AspectRatio);
-	}
-};
-
-/** A named bundle of crop settings used to implement crop presets. */
-USTRUCT(BlueprintType)
-struct FNamedPlateCropPreset
-{
-	GENERATED_USTRUCT_BODY()
-
-	/** Name for the preset. */
-	UPROPERTY(BlueprintReadWrite, Category = "Crop")
-	FString Name;
-
-	UPROPERTY(BlueprintReadWrite, Category = "Crop")
-	FPlateCropSettings CropSettings;
-};
-
-/** Supported methods for focusing the camera. */
-UENUM()
-enum class ECameraFocusMethod : uint8
-{
-	/** Don't override, ie. allow post process volume settings to persist. */
-	DoNotOverride,
-
-	/** Allows for specifying or animating exact focus distances. */
-	Manual,
-
-	/** Locks focus to specific object. */
-	Tracking,
-
-	/** Disable depth of field entirely. */
-	Disable,
-
-	MAX UMETA(Hidden)
-};
-
-/** Settings to control tracking-focus mode. */
-USTRUCT(BlueprintType)
-struct FCameraTrackingFocusSettings
-{
-	GENERATED_USTRUCT_BODY()
-
-	/** Focus distance will be tied to this actor's location. */
-	UPROPERTY(Interp, EditAnywhere, BlueprintReadWrite, Category = "Tracking Focus")
-	TSoftObjectPtr<AActor> ActorToTrack;
-
-	/** Offset from actor position to track. Relative to actor if tracking an actor, relative to world otherwise. */
-	UPROPERTY(Interp, EditAnywhere, BlueprintReadWrite, Category = "Tracking Focus")
-	FVector RelativeOffset;
-
-	/** True to draw a debug representation of the tracked position. */
-	UPROPERTY(Transient, EditAnywhere, BlueprintReadWrite, Category = "Tracking Focus")
-	uint8 bDrawDebugTrackingFocusPoint : 1;
-
-	FCameraTrackingFocusSettings()
-		: RelativeOffset(ForceInitToZero),
-		bDrawDebugTrackingFocusPoint(false)
-	{}
-};
-
-/** Settings to control camera focus */
-USTRUCT(BlueprintType)
-struct FCameraFocusSettings
-{
-	GENERATED_USTRUCT_BODY()
-
-	/** Which method to use to handle camera focus */
-	UPROPERTY(EditAnywhere, BlueprintReadWrite, Category = "Focus Method")
-	ECameraFocusMethod FocusMethod;
-	
-	/** Manually-controlled focus distance (manual focus mode only) */
-	UPROPERTY(Interp, EditAnywhere, BlueprintReadWrite, Category = "Manual Focus Settings", meta=(Units=cm))
-	float ManualFocusDistance;
-
-	/** Settings to control tracking focus (tracking focus mode only) */
-	UPROPERTY(EditAnywhere, BlueprintReadWrite, Category = "Tracking Focus Settings")
-	FCameraTrackingFocusSettings TrackingFocusSettings;
-
-#if WITH_EDITORONLY_DATA
-	/** True to draw a translucent plane at the current focus depth, for easy tweaking. */
-	UPROPERTY(Transient, EditAnywhere, Category = "Focus Settings")
-	uint8 bDrawDebugFocusPlane : 1;
-
-	/** For customizing the focus plane color, in case the default doesn't show up well in your scene. */
-	UPROPERTY(EditAnywhere, Category = "Focus Settings", meta = (EditCondition = "bDrawDebugFocusPlane"))
-	FColor DebugFocusPlaneColor;
-#endif 
-
-	/** True to use interpolation to smooth out changes in focus distance, false for focus distance changes to be instantaneous. */
-	UPROPERTY(EditAnywhere, BlueprintReadWrite, Category = "Focus Settings")
-	uint8 bSmoothFocusChanges : 1;
-	
-	/** Controls interpolation speed when smoothing focus distance changes. Ignored if bSmoothFocusChanges is false. */
-	UPROPERTY(EditAnywhere, BlueprintReadWrite, Category = "Focus Settings")
-	float FocusSmoothingInterpSpeed;
-
-	/** Additional focus depth offset, used for manually tweaking if your chosen focus method needs adjustment */
-	UPROPERTY(Interp, EditAnywhere, BlueprintReadWrite, Category = "Focus Settings")
-	float FocusOffset;
-
-	FCameraFocusSettings() : 
-		FocusMethod(ECameraFocusMethod::Manual),
-		ManualFocusDistance(100000.f),
-		TrackingFocusSettings(),
-#if WITH_EDITORONLY_DATA
-		bDrawDebugFocusPlane(false),
-		DebugFocusPlaneColor(102, 26, 204, 153),		// purple
-#endif
-		bSmoothFocusChanges(false),
-		FocusSmoothingInterpSpeed(8.f),
-		FocusOffset(0.f)
-	{}
-};
-
 /**
  * A specialized version of a camera component, geared toward cinematic usage.
  */
@@ -344,8 +93,9 @@ public:
 	FString GetFilmbackPresetName() const;
 
 	/** Returns the name of the default filmback preset. */
+	UE_DEPRECATED(5.1, "This function has been deprecated, please access the presets via the CineCameraSettings object.")
 	UFUNCTION(BlueprintPure, Category = "Cine Camera")
-	FString GetDefaultFilmbackPresetName() const { return DefaultFilmbackPreset; }
+	FString GetDefaultFilmbackPresetName() const;
 
 	/** Set the current preset settings by preset name. */
 	UFUNCTION(BlueprintCallable, Category = "Cine Camera")
@@ -368,22 +118,23 @@ public:
 	void SetCropPresetByName(const FString& InPresetName);
 
 	/** Returns a copy of the list of available filmback presets. */
+	UE_DEPRECATED(5.1, "This function has been deprecated, please access the presets via the CineCameraSettings object.")
 	UFUNCTION(BlueprintCallable, Category = "Cine Camera")
 	static TArray<FNamedFilmbackPreset> GetFilmbackPresetsCopy();
 
 	/** Returns a copy of the list of available lens presets. */
+	UE_DEPRECATED(5.1, "This function has been deprecated, please access the presets via the CineCameraSettings object.")
 	UFUNCTION(BlueprintCallable, Category = "Cine Camera")
 	static TArray<FNamedLensPreset> GetLensPresetsCopy();
 
 	/** Returns a list of available filmback presets. */
+	UE_DEPRECATED(5.1, "This function has been deprecated, please access the presets via the CineCameraSettings object.")
 	static TArray<FNamedFilmbackPreset> const& GetFilmbackPresets();
 	
 	/** Returns a list of available lens presets. */
+	UE_DEPRECATED(5.1, "This function has been deprecated, please access the presets via the CineCameraSettings object.")
 	static TArray<FNamedLensPreset> const& GetLensPresets();
 
-	/** Returns a list of available crop presets. */
-	static TArray<FNamedPlateCropPreset> const& GetCropPresets();
-
 #if WITH_EDITOR
 	/** Update the debug focus plane position and orientation. */
 	void UpdateDebugFocusPlane();
@@ -437,38 +188,37 @@ protected:
 #endif
 
 	/** List of available filmback presets */
+	UE_DEPRECATED(5.1, "This property is now located on the UCineCameraSettings object")
 	UPROPERTY(config)
 	TArray<FNamedFilmbackPreset> FilmbackPresets;
 
 	/** List of available lens presets */
+	UE_DEPRECATED(5.1, "This property is now located on the UCineCameraSettings object")
 	UPROPERTY(config)
 	TArray<FNamedLensPreset> LensPresets;
-	
-	/** List of available lens presets */
-	UPROPERTY(config)
-	TArray<FNamedPlateCropPreset> CropPresets;
 
 	/** Deprecated. See DefaultFilmbackPreset */
+	UE_DEPRECATED(5.1, "This property has been removed and fully replaced by DefaultFilmbackPreset on the UCineCameraSettings object")
 	UPROPERTY(config)
 	FString DefaultFilmbackPresetName_DEPRECATED;
 
 	/** Name of the default filmback preset */
+	UE_DEPRECATED(5.1, "This property is now located on the UCineCameraSettings object")
 	UPROPERTY(config)
 	FString DefaultFilmbackPreset;
 
 	/** Name of the default lens preset */
+	UE_DEPRECATED(5.1, "This property is now located on the UCineCameraSettings object")
 	UPROPERTY(config)
 	FString DefaultLensPresetName;
 
-	/** Name of the default crop preset */
-	UPROPERTY(config)
-	FString DefaultCropPresetName;
-
 	/** Default focal length (will be constrained by default lens) */
+	UE_DEPRECATED(5.1, "This property is now located on the UCineCameraSettings object")
 	UPROPERTY(config)
 	float DefaultLensFocalLength;
 	
 	/** Default aperture (will be constrained by default lens) */
+	UE_DEPRECATED(5.1, "This property is now located on the UCineCameraSettings object")
 	UPROPERTY(config)
 	float DefaultLensFStop;
 
diff --git a/Engine/Source/Runtime/CinematicCamera/Public/CineCameraSettings.h b/Engine/Source/Runtime/CinematicCamera/Public/CineCameraSettings.h
new file mode 100644
--- /dev/null
+++ b/Engine/Source/Runtime/CinematicCamera/Public/CineCameraSettings.h
@@ -0,0 +1,354 @@
+// Copyright Epic Games, Inc. All Rights Reserved.
+
+#pragma once
+
+#include "CoreMinimal.h"
+#include "Engine/DeveloperSettings.h"
+#include "Engine/Scene.h"
+#include "Templates/SharedPointer.h"
+#include "CineCameraSettings.generated.h"
+
+class SNotificationItem;
+
+/** #note, this struct has a details customization in CameraFilmbackSettingsCustomization.cpp/h */
+USTRUCT(BlueprintType)
+struct FCameraFilmbackSettings
+{
+	GENERATED_USTRUCT_BODY()
+
+	/** Horizontal size of filmback or digital sensor, in mm. */
+	UPROPERTY(Interp, EditAnywhere, BlueprintReadWrite, Category = "Filmback", meta = (ClampMin = "0.001", ForceUnits = mm))
+	float SensorWidth;
+
+	/** Vertical size of filmback or digital sensor, in mm. */
+	UPROPERTY(Interp, EditAnywhere, BlueprintReadWrite, Category = "Filmback", meta = (ClampMin = "0.001", ForceUnits = mm))
+	float SensorHeight;
+
+	/** Read-only. Computed from Sensor dimensions. */
+	UPROPERTY(Interp, VisibleAnywhere, BlueprintReadOnly, Category = "Filmback")
+	float SensorAspectRatio;
+
+	bool operator==(const FCameraFilmbackSettings& Other) const
+	{
+		return (SensorWidth == Other.SensorWidth)
+			&& (SensorHeight == Other.SensorHeight);
+	}
+
+	bool operator!=(const FCameraFilmbackSettings& Other) const
+	{
+		return !operator==(Other);
+	}
+
+	FCameraFilmbackSettings()
+		: SensorWidth(24.89f)
+		, SensorHeight(18.67f)
+		, SensorAspectRatio(1.33f)
+	{
+	}
+};
+
+/** A named bundle of filmback settings used to implement filmback presets */
+USTRUCT(BlueprintType)
+struct FNamedFilmbackPreset
+{
+	GENERATED_USTRUCT_BODY()
+
+	/** Name for the preset. */
+	UPROPERTY(BlueprintReadWrite, EditAnywhere, Category = "Filmback")
+	FString Name;
+
+	UPROPERTY(BlueprintReadWrite, EditAnywhere, Category = "Filmback")
+	FCameraFilmbackSettings FilmbackSettings;
+};
+
+/** 
+ * #note, this struct has a details customization in CameraLensSettingsCustomization.cpp/h
+ */
+USTRUCT(BlueprintType)
+struct FCameraLensSettings
+{
+	GENERATED_USTRUCT_BODY()
+
+	/** Default constructor, initializing with default values */
+	FCameraLensSettings()
+		: MinFocalLength(50.f)
+		, MaxFocalLength(50.f)
+		, MinFStop(2.f)
+		, MaxFStop(2.f)
+		, MinimumFocusDistance(15.f)
+		, DiaphragmBladeCount(FPostProcessSettings::kDefaultDepthOfFieldBladeCount)
+	{
+	}
+
+	/** Minimum focal length for this lens */
+	UPROPERTY(EditAnywhere, BlueprintReadWrite, Category = "Lens", meta = (ForceUnits = mm, ClampMin = "0.001"))
+	float MinFocalLength = 0.f;
+
+	/** Maximum focal length for this lens */
+	UPROPERTY(EditAnywhere, BlueprintReadWrite, Category = "Lens", meta = (ForceUnits = mm, ClampMin = "0.001"))
+	float MaxFocalLength = 0.f;
+
+	/** Minimum aperture for this lens (e.g. 2.8 for an f/2.8 lens) */
+	UPROPERTY(EditAnywhere, BlueprintReadWrite, Category = "Lens")
+	float MinFStop = 0.f;
+
+	/** Maximum aperture for this lens (e.g. 2.8 for an f/2.8 lens) */
+	UPROPERTY(EditAnywhere, BlueprintReadWrite, Category = "Lens")
+	float MaxFStop = 0.f;
+
+	/** Shortest distance this lens can focus on. */
+	UPROPERTY(EditAnywhere, BlueprintReadWrite, Category = "Lens", meta = (ForceUnits = mm))
+	float MinimumFocusDistance = 0.f;
+
+	/** Squeeze factor for anamorphic lenses. */
+	UPROPERTY(EditAnywhere, BlueprintReadWrite, Category = "Lens", meta = (ClampMin = "1", ClampMax = "2"))
+	float SqueezeFactor = 1.f;
+
+	/** Number of blades of diaphragm. */
+	UPROPERTY(EditAnywhere, BlueprintReadWrite, Category = "Lens", meta = (ClampMin = "4", ClampMax = "16"))
+	int32 DiaphragmBladeCount = 0;
+
+	bool operator==(const FCameraLensSettings& Other) const
+	{
+		return (MinFocalLength == Other.MinFocalLength)
+			&& (MaxFocalLength == Other.MaxFocalLength)
+			&& (MinFStop == Other.MinFStop)
+			&& (MaxFStop == Other.MaxFStop)
+			&& (MinimumFocusDistance == Other.MinimumFocusDistance)
+			&& (SqueezeFactor == Other.SqueezeFactor)
+			&& (DiaphragmBladeCount == Other.DiaphragmBladeCount);
+	}
+};
+
+/** A named bundle of lens settings used to implement lens presets. */
+USTRUCT(BlueprintType)
+struct FNamedLensPreset
+{
+	GENERATED_USTRUCT_BODY()
+
+	/** Name for the preset. */
+	UPROPERTY(BlueprintReadWrite, EditAnywhere, Category = "Lens")
+	FString Name;
+
+	UPROPERTY(BlueprintReadWrite, EditAnywhere, Category = "Lens")
+	FCameraLensSettings LensSettings;
+};
+
+/** 
+ * #note, this struct has a details customization in CameraCropSettingsCustomization.cpp/h
+ */
+USTRUCT(BlueprintType)
+struct FPlateCropSettings
+{
+	GENERATED_USTRUCT_BODY()
+
+	UPROPERTY(EditAnywhere, BlueprintReadWrite, Category = "Crop", meta=(DisplayName = "Cropped Aspect Ratio"))
+	float AspectRatio = 0.f;
+
+	bool operator==(const FPlateCropSettings& Other) const
+	{
+		return (AspectRatio == Other.AspectRatio);
+	}
+};
+
+/** A named bundle of crop settings used to implement crop presets. */
+USTRUCT(BlueprintType)
+struct FNamedPlateCropPreset
+{
+	GENERATED_USTRUCT_BODY()
+
+	/** Name for the preset. */
+	UPROPERTY(BlueprintReadWrite, EditAnywhere, Category = "Crop")
+	FString Name;
+
+	UPROPERTY(BlueprintReadWrite, EditAnywhere, Category = "Crop")
+	FPlateCropSettings CropSettings;
+};
+
+/** Supported methods for focusing the camera. */
+UENUM()
+enum class ECameraFocusMethod : uint8
+{
+	/** Don't override, ie. allow post process volume settings to persist. */
+	DoNotOverride,
+
+	/** Allows for specifying or animating exact focus distances. */
+	Manual,
+
+	/** Locks focus to specific object. */
+	Tracking,
+
+	/** Disable depth of field entirely. */
+	Disable,
+
+	MAX UMETA(Hidden)
+};
+
+/** Settings to control tracking-focus mode. */
+USTRUCT(BlueprintType)
+struct FCameraTrackingFocusSettings
+{
+	GENERATED_USTRUCT_BODY()
+
+	/** Focus distance will be tied to this actor's location. */
+	UPROPERTY(Interp, EditAnywhere, BlueprintReadWrite, Category = "Tracking Focus")
+	TSoftObjectPtr<AActor> ActorToTrack;
+
+	/** Offset from actor position to track. Relative to actor if tracking an actor, relative to world otherwise. */
+	UPROPERTY(Interp, EditAnywhere, BlueprintReadWrite, Category = "Tracking Focus")
+	FVector RelativeOffset;
+
+	/** True to draw a debug representation of the tracked position. */
+	UPROPERTY(Transient, EditAnywhere, BlueprintReadWrite, Category = "Tracking Focus")
+	uint8 bDrawDebugTrackingFocusPoint : 1;
+
+	FCameraTrackingFocusSettings()
+		: RelativeOffset(ForceInitToZero),
+		bDrawDebugTrackingFocusPoint(false)
+	{}
+};
+
+/** Settings to control camera focus */
+USTRUCT(BlueprintType)
+struct FCameraFocusSettings
+{
+	GENERATED_USTRUCT_BODY()
+
+	/** Which method to use to handle camera focus */
+	UPROPERTY(EditAnywhere, BlueprintReadWrite, Category = "Focus Method")
+	ECameraFocusMethod FocusMethod;
+	
+	/** Manually-controlled focus distance (manual focus mode only) */
+	UPROPERTY(Interp, EditAnywhere, BlueprintReadWrite, Category = "Manual Focus Settings", meta=(Units=cm))
+	float ManualFocusDistance;
+
+	/** Settings to control tracking focus (tracking focus mode only) */
+	UPROPERTY(EditAnywhere, BlueprintReadWrite, Category = "Tracking Focus Settings")
+	FCameraTrackingFocusSettings TrackingFocusSettings;
+
+#if WITH_EDITORONLY_DATA
+	/** True to draw a translucent plane at the current focus depth, for easy tweaking. */
+	UPROPERTY(Transient, EditAnywhere, Category = "Focus Settings")
+	uint8 bDrawDebugFocusPlane : 1;
+
+	/** For customizing the focus plane color, in case the default doesn't show up well in your scene. */
+	UPROPERTY(EditAnywhere, Category = "Focus Settings", meta = (EditCondition = "bDrawDebugFocusPlane"))
+	FColor DebugFocusPlaneColor;
+#endif 
+
+	/** True to use interpolation to smooth out changes in focus distance, false for focus distance changes to be instantaneous. */
+	UPROPERTY(EditAnywhere, BlueprintReadWrite, Category = "Focus Settings")
+	uint8 bSmoothFocusChanges : 1;
+	
+	/** Controls interpolation speed when smoothing focus distance changes. Ignored if bSmoothFocusChanges is false. */
+	UPROPERTY(EditAnywhere, BlueprintReadWrite, Category = "Focus Settings")
+	float FocusSmoothingInterpSpeed;
+
+	/** Additional focus depth offset, used for manually tweaking if your chosen focus method needs adjustment */
+	UPROPERTY(Interp, EditAnywhere, BlueprintReadWrite, Category = "Focus Settings")
+	float FocusOffset;
+
+	FCameraFocusSettings() : 
+		FocusMethod(ECameraFocusMethod::Manual),
+		ManualFocusDistance(100000.f),
+		TrackingFocusSettings(),
+#if WITH_EDITORONLY_DATA
+		bDrawDebugFocusPlane(false),
+		DebugFocusPlaneColor(102, 26, 204, 153),		// purple
+#endif
+		bSmoothFocusChanges(false),
+		FocusSmoothingInterpSpeed(8.f),
+		FocusOffset(0.f)
+	{}
+};
+
+UCLASS(config=Engine, defaultconfig, meta=(DisplayName="Cinematic Camera"))
+class CINEMATICCAMERA_API UCineCameraSettings : public UDeveloperSettings
+{
+	GENERATED_BODY()
+
+public:
+	virtual void PostInitProperties() override;
+
+	/** Name of the default lens preset */
+	UPROPERTY(config, EditAnywhere, BlueprintReadWrite, BlueprintSetter=SetDefaultLensPresetName, Category=Lens, meta=(GetOptions=GetLensPresetNames))
+	FString DefaultLensPresetName;
+	
+	/** Default focal length (will be constrained by default lens) */
+	UPROPERTY(config, EditAnywhere, BlueprintReadWrite, BlueprintSetter=SetDefaultLensFocalLength, Category=Lens)
+	float DefaultLensFocalLength;
+	
+	/** Default aperture (will be constrained by default lens) */
+	UPROPERTY(config, EditAnywhere, BlueprintReadWrite, BlueprintSetter=SetDefaultLensFStop, Category=Lens)
+	float DefaultLensFStop;
+	
+	/** List of available lens presets */
+	UPROPERTY(config, EditAnywhere, BlueprintReadWrite, BlueprintSetter=SetLensPresets, Category=Lens, meta=(TitleProperty=Name))
+	TArray<FNamedLensPreset> LensPresets;
+
+	/** Name of the default filmback preset */
+	UPROPERTY(config, EditAnywhere, BlueprintReadWrite, BlueprintSetter=SetDefaultfilmbackPreset, Category=Filmback, meta=(GetOptions=GetFilmbackPresetNames))
+	FString DefaultFilmbackPreset;
+
+	/** List of available filmback presets */
+	UPROPERTY(config, EditAnywhere, BlueprintReadWrite, BlueprintSetter=SetFilmbackPresets, Category=Filmback, meta=(TitleProperty=Name))
+	TArray<FNamedFilmbackPreset> FilmbackPresets;
+
+	/** Name of the default crop preset */
+	UPROPERTY(config, EditAnywhere, BlueprintReadWrite, BlueprintSetter=SetDefaultCropPresetName, Category=Crop, meta=(GetOptions=GetCropPresetNames))
+	FString DefaultCropPresetName;
+	
+	/** List of available crop presets */
+	UPROPERTY(config, EditAnywhere, BlueprintReadWrite, BlueprintSetter=SetCropPresets, Category=Crop, meta=(TitleProperty=Name))
+	TArray<FNamedPlateCropPreset> CropPresets;
+
+	static TArray<FNamedLensPreset> const& GetLensPresets();
+	static TArray<FNamedFilmbackPreset> const& GetFilmbackPresets();
+	static TArray<FNamedPlateCropPreset> const& GetCropPresets();
+
+private:
+	UFUNCTION(BlueprintCallable, Category=Cinematics)
+	static UCineCameraSettings* GetCineCameraSettings();
+
+	/* Internal Blueprint Setter functions that call SaveConfig after setting the variable to ensure settings persist */
+	
+	UFUNCTION(BlueprintSetter)
+	void SetDefaultLensPresetName(const FString InDefaultLensPresetName);
+
+	UFUNCTION(BlueprintSetter)
+	void SetDefaultLensFocalLength(const float InDefaultLensFocalLength);
+
+	UFUNCTION(BlueprintSetter)
+	void SetDefaultLensFStop(const float InDefaultLensFStop);
+
+	UFUNCTION(BlueprintSetter)
+	void SetLensPresets(const TArray<FNamedLensPreset>& InLensPresets);
+
+	UFUNCTION(BlueprintSetter)
+	void SetDefaultFilmbackPreset(const FString InDefaultFilmbackPreset);
+
+	UFUNCTION(BlueprintSetter)
+	void SetFilmbackPresets(const TArray<FNamedFilmbackPreset>& InFilmbackPresets);
+
+	UFUNCTION(BlueprintSetter)
+	void SetDefaultCropPresetName(const FString InDefaultCropPresetName);
+
+	UFUNCTION(BlueprintSetter)
+	void SetCropPresets(const TArray<FNamedPlateCropPreset>& InCropPresets);
+
+	/* Functions used for the GetOptions metadata */
+	
+	UFUNCTION()
+	TArray<FString> GetLensPresetNames() const;
+
+	UFUNCTION()
+	TArray<FString> GetFilmbackPresetNames() const;
+
+	UFUNCTION()
+	TArray<FString> GetCropPresetNames() const;
+
+	static const FString CineCameraConfigSection;
+	TSharedPtr<SNotificationItem> Notification;
+	void CopyOldConfigSettings();
+	void CloseNotification();
+};
