Ensured CineCameraSettings notification is only displayed in Editor

diff --git a/Engine/Source/Runtime/CinematicCamera/Private/CineCameraSettings.cpp b/Engine/Source/Runtime/CinematicCamera/Private/CineCameraSettings.cpp
--- a/Engine/Source/Runtime/CinematicCamera/Private/CineCameraSettings.cpp
+++ b/Engine/Source/Runtime/CinematicCamera/Private/CineCameraSettings.cpp
@@ -15,7 +15,9 @@ const FString UCineCameraSettings::CineCameraConfigSection = TEXT("/Script/Cinem
 void UCineCameraSettings::PostInitProperties()
 {
 	Super::PostInitProperties();
-	
+
+// The Fixup Notifications should only be displayed in Editor
+#if WITH_EDITOR
 	if (GConfig && GConfig->DoesSectionExist(*CineCameraConfigSection, GEngineIni))
 	{
 		FNotificationInfo NotificationInfo(LOCTEXT("UpdateSettingsNotification", "CineCamera Settings were found in an old config location. Do you want to attempt to automatically merge them to the new location?"));
@@ -31,6 +33,7 @@ void UCineCameraSettings::PostInitProperties()
 
 		Notification = FSlateNotificationManager::Get().AddNotification(NotificationInfo);
 	}
+#endif
 }
 
 void UCineCameraSettings::SetDefaultLensPresetName(const FString InDefaultLensPresetName)
