CineCameraComponent now correctly respects the default Crop Settings

diff --git a/Engine/Source/Runtime/CinematicCamera/Private/CineCameraComponent.cpp b/Engine/Source/Runtime/CinematicCamera/Private/CineCameraComponent.cpp
--- a/Engine/Source/Runtime/CinematicCamera/Private/CineCameraComponent.cpp
+++ b/Engine/Source/Runtime/CinematicCamera/Private/CineCameraComponent.cpp
@@ -44,6 +44,7 @@ UCineCameraComponent::UCineCameraComponent()
 	SetFilmbackPresetByNameInternal(CineCameraSettings->DefaultFilmbackPreset, Filmback);
 	SetFilmbackPresetByNameInternal(CineCameraSettings->DefaultFilmbackPreset, FilmbackSettings_DEPRECATED);
 	SetLensPresetByNameInternal(CineCameraSettings->DefaultLensPresetName);
+	SetCropPresetByNameInternal(CineCameraSettings->DefaultCropPresetName);
 	// other lens defaults
 	CurrentAperture = CineCameraSettings->DefaultLensFStop;
 	CurrentFocalLength = CineCameraSettings->DefaultLensFocalLength;
