Fixed bug where Sensor Aspect Ratio does not update when setting Filmback Presets in the Cine Camera Settings

diff --git a/Engine/Source/Runtime/CinematicCamera/Private/CineCameraComponent.cpp b/Engine/Source/Runtime/CinematicCamera/Private/CineCameraComponent.cpp
--- a/Engine/Source/Runtime/CinematicCamera/Private/CineCameraComponent.cpp
+++ b/Engine/Source/Runtime/CinematicCamera/Private/CineCameraComponent.cpp
@@ -478,7 +478,7 @@ void UCineCameraComponent::RecalcDerivedData()
 	FocusSettings.ManualFocusDistance = FMath::Max(FocusSettings.ManualFocusDistance, MinFocusDistInWorldUnits);
 
 	FieldOfView = GetHorizontalFieldOfView();
-	Filmback.SensorAspectRatio = (Filmback.SensorHeight > 0.f) ? (Filmback.SensorWidth / Filmback.SensorHeight) : 0.f;
+	Filmback.RecalcSensorAspectRatio();
 	AspectRatio = Filmback.SensorAspectRatio * LensSettings.SqueezeFactor;
 	if (CropSettings.AspectRatio > 0.0f)
 	{
diff --git a/Engine/Source/Runtime/CinematicCamera/Private/CineCameraSettings.cpp b/Engine/Source/Runtime/CinematicCamera/Private/CineCameraSettings.cpp
--- a/Engine/Source/Runtime/CinematicCamera/Private/CineCameraSettings.cpp
+++ b/Engine/Source/Runtime/CinematicCamera/Private/CineCameraSettings.cpp
@@ -34,8 +34,21 @@ void UCineCameraSettings::PostInitProperties()
 		Notification = FSlateNotificationManager::Get().AddNotification(NotificationInfo);
 	}
 #endif
+	RecalcSensorAspectRatios();
 }
 
+#if WITH_EDITOR
+void UCineCameraSettings::PostEditChangeProperty(FPropertyChangedEvent& PropertyChangedEvent)
+{
+	if (PropertyChangedEvent.GetMemberPropertyName() == GET_MEMBER_NAME_CHECKED(UCineCameraSettings, FilmbackPresets))
+	{
+		RecalcSensorAspectRatios();
+	}
+	
+	Super::PostEditChangeProperty(PropertyChangedEvent);
+}
+#endif
+
 void UCineCameraSettings::SetDefaultLensPresetName(const FString InDefaultLensPresetName)
 {
 	DefaultLensPresetName = InDefaultLensPresetName;
@@ -69,6 +82,7 @@ void UCineCameraSettings::SetDefaultFilmbackPreset(const FString InDefaultFilmba
 void UCineCameraSettings::SetFilmbackPresets(const TArray<FNamedFilmbackPreset>& InFilmbackPresets)
 {
 	FilmbackPresets = InFilmbackPresets;
+	RecalcSensorAspectRatios();
 	SaveConfig();
 }
 
@@ -152,6 +166,14 @@ void UCineCameraSettings::CloseNotification()
 	FSlateNotificationManager::Get().AddNotification(NotificationInfo);
 }
 
+void UCineCameraSettings::RecalcSensorAspectRatios()
+{
+	for (FNamedFilmbackPreset& FilmbackPreset : FilmbackPresets)
+	{
+		FilmbackPreset.FilmbackSettings.RecalcSensorAspectRatio();	
+	}
+}
+
 void UCineCameraSettings::CopyOldConfigSettings()
 {	
 	const FString SettingsConfigSection(TEXT("/Script/CinematicCamera.CineCameraSettings"));
diff --git a/Engine/Source/Runtime/CinematicCamera/Public/CineCameraSettings.h b/Engine/Source/Runtime/CinematicCamera/Public/CineCameraSettings.h
--- a/Engine/Source/Runtime/CinematicCamera/Public/CineCameraSettings.h
+++ b/Engine/Source/Runtime/CinematicCamera/Public/CineCameraSettings.h
@@ -39,6 +39,11 @@ struct FCameraFilmbackSettings
 		return !operator==(Other);
 	}
 
+	void RecalcSensorAspectRatio()
+	{
+		SensorAspectRatio = (SensorHeight > 0.f) ? (SensorWidth / SensorHeight) : 0.f;
+	}
+
 	FCameraFilmbackSettings()
 		: SensorWidth(24.89f)
 		, SensorHeight(18.67f)
@@ -270,6 +275,10 @@ class CINEMATICCAMERA_API UCineCameraSettings : public UDeveloperSettings
 public:
 	virtual void PostInitProperties() override;
 
+#if WITH_EDITOR
+	virtual void PostEditChangeProperty(FPropertyChangedEvent& PropertyChangedEvent) override;
+#endif
+
 	/** Name of the default lens preset */
 	UPROPERTY(config, EditAnywhere, BlueprintReadWrite, BlueprintSetter=SetDefaultLensPresetName, Category=Lens, meta=(GetOptions=GetLensPresetNames))
 	FString DefaultLensPresetName;
@@ -351,4 +360,7 @@ private:
 	TSharedPtr<SNotificationItem> Notification;
 	void CopyOldConfigSettings();
 	void CloseNotification();
+
+	// Sensor aspect ratio is derived from Sensor Height and Width so needs to be updated when either property is updated
+	void RecalcSensorAspectRatios();
 };
