Added functions to get CineCameraSettings presets by name

diff --git a/Engine/Source/Runtime/CinematicCamera/Private/CineCameraSettings.cpp b/Engine/Source/Runtime/CinematicCamera/Private/CineCameraSettings.cpp
--- a/Engine/Source/Runtime/CinematicCamera/Private/CineCameraSettings.cpp
+++ b/Engine/Source/Runtime/CinematicCamera/Private/CineCameraSettings.cpp
@@ -113,6 +113,39 @@ TArray<FNamedPlateCropPreset> const& UCineCameraSettings::GetCropPresets()
 	return GetDefault<UCineCameraSettings>()->CropPresets;
 }
 
+bool UCineCameraSettings::GetLensPresetByName(const FString PresetName, FCameraLensSettings& LensSettings)
+{
+	FNamedLensPreset* NamedLensPreset = LensPresets.FindByPredicate([PresetName](const FNamedLensPreset& Preset)
+	{
+		return Preset.Name == PresetName;
+	});
+
+	LensSettings = NamedLensPreset ? NamedLensPreset->LensSettings : FCameraLensSettings();
+	return NamedLensPreset != nullptr;
+}
+
+bool UCineCameraSettings::GetFilmbackPresetByName(const FString PresetName, FCameraFilmbackSettings& FilmbackSettings)
+{
+	FNamedFilmbackPreset* NamedFilmbackPreset = FilmbackPresets.FindByPredicate([PresetName](const FNamedFilmbackPreset& Preset)
+	{
+		return Preset.Name == PresetName;
+	});
+
+	FilmbackSettings = NamedFilmbackPreset ? NamedFilmbackPreset->FilmbackSettings : FCameraFilmbackSettings();
+	return NamedFilmbackPreset != nullptr;
+}
+
+bool UCineCameraSettings::GetCropPresetByName(const FString PresetName, FPlateCropSettings& CropSettings)
+{
+	FNamedPlateCropPreset* NamedCropPreset = CropPresets.FindByPredicate([PresetName](const FNamedPlateCropPreset& Preset)
+	{
+		return Preset.Name == PresetName;
+	});
+
+	CropSettings = NamedCropPreset ? NamedCropPreset->CropSettings : FPlateCropSettings();
+	return NamedCropPreset != nullptr;
+}
+
 UCineCameraSettings* UCineCameraSettings::GetCineCameraSettings()
 {
 	return GetMutableDefault<UCineCameraSettings>();
diff --git a/Engine/Source/Runtime/CinematicCamera/Public/CineCameraSettings.h b/Engine/Source/Runtime/CinematicCamera/Public/CineCameraSettings.h
--- a/Engine/Source/Runtime/CinematicCamera/Public/CineCameraSettings.h
+++ b/Engine/Source/Runtime/CinematicCamera/Public/CineCameraSettings.h
@@ -315,6 +315,21 @@ public:
 	static TArray<FNamedFilmbackPreset> const& GetFilmbackPresets();
 	static TArray<FNamedPlateCropPreset> const& GetCropPresets();
 
+	// Gets the Lens settings associated with a given preset name
+	// Returns true if a preset with the given name was found
+	UFUNCTION(BlueprintCallable, Category=Cinematics, meta=(ReturnDisplayName = "Success"))
+	bool GetLensPresetByName(const FString PresetName, FCameraLensSettings& LensSettings);
+
+	// Gets the Filmback settings associated with a given preset name
+	// Returns true if a preset with the given name was found
+	UFUNCTION(BlueprintCallable, Category=Cinematics, meta=(ReturnDisplayName = "Success"))
+	bool GetFilmbackPresetByName(const FString PresetName, FCameraFilmbackSettings& FilmbackSettings);
+
+	// Gets the Crop settings associated with a given preset name
+	// Returns true if a preset with the given name was found
+	UFUNCTION(BlueprintCallable, Category=Cinematics, meta=(ReturnDisplayName = "Success"))
+	bool GetCropPresetByName(const FString PresetName, FPlateCropSettings& CropSettings);
+
 private:
 	UFUNCTION(BlueprintCallable, Category=Cinematics)
 	static UCineCameraSettings* GetCineCameraSettings();
