Fix for derived CineCamera data not updating when modifying properties via Blueprints

diff --git a/Engine/Source/Runtime/CinematicCamera/Private/CineCameraComponent.cpp b/Engine/Source/Runtime/CinematicCamera/Private/CineCameraComponent.cpp
--- a/Engine/Source/Runtime/CinematicCamera/Private/CineCameraComponent.cpp
+++ b/Engine/Source/Runtime/CinematicCamera/Private/CineCameraComponent.cpp
@@ -591,6 +591,36 @@ FText UCineCameraComponent::GetFilmbackText() const
 }
 #endif
 
+void UCineCameraComponent::SetFilmback(const FCameraFilmbackSettings& NewFilmback)
+{
+	Filmback = NewFilmback;
+	RecalcDerivedData();
+}
+
+void UCineCameraComponent::SetLensSettings(const FCameraLensSettings& NewLensSettings)
+{
+	LensSettings = NewLensSettings;
+	RecalcDerivedData();
+}
+
+void UCineCameraComponent::SetFocusSettings(const FCameraFocusSettings& NewFocusSettings)
+{
+	FocusSettings = NewFocusSettings;
+	RecalcDerivedData();
+}
+
+void UCineCameraComponent::SetCropSettings(const FPlateCropSettings& NewCropSettings)
+{
+	CropSettings = NewCropSettings;
+	RecalcDerivedData();
+}
+
+void UCineCameraComponent::SetCurrentAperture(const float NewCurrentAperture)
+{
+	CurrentAperture = NewCurrentAperture;
+	RecalcDerivedData();
+}
+
 #if WITH_EDITORONLY_DATA
 void UCineCameraComponent::UpdateDebugFocusPlane()
 {
diff --git a/Engine/Source/Runtime/CinematicCamera/Public/CineCameraComponent.h b/Engine/Source/Runtime/CinematicCamera/Public/CineCameraComponent.h
--- a/Engine/Source/Runtime/CinematicCamera/Public/CineCameraComponent.h
+++ b/Engine/Source/Runtime/CinematicCamera/Public/CineCameraComponent.h
@@ -34,28 +34,43 @@ public:
 	FCameraFilmbackSettings FilmbackSettings_DEPRECATED;
 
 	/** Controls the filmback of the camera. */
-	UPROPERTY(Interp, EditAnywhere, BlueprintReadWrite, Category = "Current Camera Settings")
+	UPROPERTY(Interp, BlueprintSetter = SetFilmback, EditAnywhere, BlueprintReadWrite, Category = "Current Camera Settings")
 	FCameraFilmbackSettings Filmback;
 
+	UFUNCTION(BlueprintSetter)
+	void SetFilmback(const FCameraFilmbackSettings& NewFilmback);
+
 	/** Controls the camera's lens. */
-	UPROPERTY(EditAnywhere, BlueprintReadWrite, Category = "Current Camera Settings")
+	UPROPERTY(EditAnywhere, BlueprintSetter = SetLensSettings, BlueprintReadWrite, Category = "Current Camera Settings")
 	FCameraLensSettings LensSettings;
 
+	UFUNCTION(BlueprintSetter)
+	void SetLensSettings(const FCameraLensSettings& NewLensSettings);
+
 	/** Controls the camera's focus. */
-	UPROPERTY(EditAnywhere, BlueprintReadWrite, Category = "Current Camera Settings")
+	UPROPERTY(EditAnywhere, BlueprintSetter = SetFocusSettings, BlueprintReadWrite, Category = "Current Camera Settings")
 	FCameraFocusSettings FocusSettings;
 
+	UFUNCTION(BlueprintSetter)
+	void SetFocusSettings(const FCameraFocusSettings& NewFocusSettings);
+	
 	/** Controls the crop settings. */
-	UPROPERTY(EditAnywhere, BlueprintReadWrite, Category = "Current Camera Settings")
+	UPROPERTY(EditAnywhere, BlueprintSetter = SetCropSettings, BlueprintReadWrite, Category = "Current Camera Settings")
 	FPlateCropSettings CropSettings;
 
+	UFUNCTION(BlueprintSetter)
+	void SetCropSettings(const FPlateCropSettings& NewCropSettings);
+
 	/** Current focal length of the camera (i.e. controls FoV, zoom) */
 	UPROPERTY(Interp, BlueprintSetter = SetCurrentFocalLength, EditAnywhere, BlueprintReadWrite, Category = "Current Camera Settings")
 	float CurrentFocalLength;
 
 	/** Current aperture, in terms of f-stop (e.g. 2.8 for f/2.8) */
-	UPROPERTY(Interp, EditAnywhere, BlueprintReadWrite, Category = "Current Camera Settings")
+	UPROPERTY(Interp, BlueprintSetter = SetCurrentAperture, EditAnywhere, BlueprintReadWrite, Category = "Current Camera Settings")
 	float CurrentAperture;
+
+	UFUNCTION(BlueprintSetter)
+	void SetCurrentAperture(const float NewCurrentAperture);
 	
 	/** Read-only. Control this value via FocusSettings. */
 	UPROPERTY(VisibleAnywhere, BlueprintReadOnly, Category = "Current Camera Settings")
