Cine Camera Actor: Fixed the issue of when viewport dissapears when Custom Near Clipping plane is set to 0.

#rnx

diff --git a/Engine/Source/Runtime/CinematicCamera/Public/CineCameraComponent.h b/Engine/Source/Runtime/CinematicCamera/Public/CineCameraComponent.h
--- a/Engine/Source/Runtime/CinematicCamera/Public/CineCameraComponent.h
+++ b/Engine/Source/Runtime/CinematicCamera/Public/CineCameraComponent.h
@@ -80,7 +80,7 @@ public:
 	uint32 bOverride_CustomNearClippingPlane:1;
 
 	/** Set bOverride_CustomNearClippingPlane to true if you want to use a custom clipping plane instead of GNearClippingPlane. */
-	UPROPERTY(EditAnywhere, BlueprintReadWrite, AdvancedDisplay, Category = "Current Camera Settings", meta = (editcondition = "bOverride_CustomNearClippingPlane"))
+	UPROPERTY(EditAnywhere, BlueprintReadWrite, AdvancedDisplay, Category = "Current Camera Settings", meta = (UIMin = "0.00001", ClampMin = "0.00001", editcondition = "bOverride_CustomNearClippingPlane"))
 	float CustomNearClippingPlane;
 
 #if WITH_EDITORONLY_DATA
