Make the override value blueprint accessible

diff --git a/Engine/Source/Runtime/CinematicCamera/Public/CineCameraComponent.h b/Engine/Source/Runtime/CinematicCamera/Public/CineCameraComponent.h
--- a/Engine/Source/Runtime/CinematicCamera/Public/CineCameraComponent.h
+++ b/Engine/Source/Runtime/CinematicCamera/Public/CineCameraComponent.h
@@ -76,7 +76,7 @@ public:
 	UPROPERTY(VisibleAnywhere, BlueprintReadOnly, Category = "Current Camera Settings")
 	float CurrentFocusDistance;
 
-	UPROPERTY(EditAnywhere, Category = "Current Camera Settings", meta = (InlineEditConditionToggle))
+	UPROPERTY(EditAnywhere, BlueprintReadWrite, AdvancedDisplay, Category = "Current Camera Settings", meta = (InlineEditConditionToggle))
 	uint32 bOverride_CustomNearClippingPlane:1;
 
 	/** Set bOverride_CustomNearClippingPlane to true if you want to use a custom clipping plane instead of GNearClippingPlane. */
