CameraRig_Rail: Made private methods and properties protected so that the child class has access. Removed orientation on the mount mesh as it is unnecessary

diff --git a/Engine/Source/Runtime/CinematicCamera/Private/CameraRig_Rail.cpp b/Engine/Source/Runtime/CinematicCamera/Private/CameraRig_Rail.cpp
--- a/Engine/Source/Runtime/CinematicCamera/Private/CameraRig_Rail.cpp
+++ b/Engine/Source/Runtime/CinematicCamera/Private/CameraRig_Rail.cpp
@@ -152,13 +152,6 @@ void ACameraRig_Rail::UpdatePreviewMeshes()
 		// make visualization of the mount follow the contour of the rail
 		if (PreviewMesh_Mount)
 		{
-			if (bLockOrientationToRail)
-			{
-				float const SplineLen = RailSplineComponent->GetSplineLength();
-				FQuat const RailRot = RailSplineComponent->GetQuaternionAtDistanceAlongSpline(CurrentPositionOnRail*SplineLen, ESplineCoordinateSpace::World);
-				PreviewMesh_Mount->SetWorldRotation(RailRot);
-			}
-			
 			PreviewMesh_Mount->SetVisibility(bShowRailVisualization);
 			PreviewMesh_Mount->SetWorldScale3D(FVector(PreviewMeshScale, PreviewMeshScale, PreviewMeshScale));
 		}
diff --git a/Engine/Source/Runtime/CinematicCamera/Public/CameraRig_Rail.h b/Engine/Source/Runtime/CinematicCamera/Public/CameraRig_Rail.h
--- a/Engine/Source/Runtime/CinematicCamera/Public/CameraRig_Rail.h
+++ b/Engine/Source/Runtime/CinematicCamera/Public/CameraRig_Rail.h
@@ -56,9 +56,9 @@ public:
 	UFUNCTION(BlueprintPure, Category = "Rail Components")
 	USplineComponent* GetRailSplineComponent() { return RailSplineComponent; }
 	
-private:
+protected:
 	/** Makes sure all components are arranged properly. Call when something changes that might affect components. */
-	void UpdateRailComponents();
+	virtual void UpdateRailComponents();
 #if WITH_EDITORONLY_DATA
 	void UpdatePreviewMeshes();
 #endif
