Cine Camera Actor: Made Custom Near Clipping plane keyable in sequencer.

#rnx

diff --git a/Engine/Source/Runtime/CinematicCamera/Private/CineCameraComponent.cpp b/Engine/Source/Runtime/CinematicCamera/Private/CineCameraComponent.cpp
--- a/Engine/Source/Runtime/CinematicCamera/Private/CineCameraComponent.cpp
+++ b/Engine/Source/Runtime/CinematicCamera/Private/CineCameraComponent.cpp
@@ -621,6 +621,11 @@ void UCineCameraComponent::SetCurrentAperture(const float NewCurrentAperture)
 	RecalcDerivedData();
 }
 
+void UCineCameraComponent::SetCustomNearClippingPlane(const float NewCustomNearClippingPlane)
+{
+	CustomNearClippingPlane = NewCustomNearClippingPlane;
+}
+
 #if WITH_EDITORONLY_DATA
 void UCineCameraComponent::UpdateDebugFocusPlane()
 {
diff --git a/Engine/Source/Runtime/CinematicCamera/Public/CineCameraComponent.h b/Engine/Source/Runtime/CinematicCamera/Public/CineCameraComponent.h
--- a/Engine/Source/Runtime/CinematicCamera/Public/CineCameraComponent.h
+++ b/Engine/Source/Runtime/CinematicCamera/Public/CineCameraComponent.h
@@ -80,9 +80,9 @@ public:
 	uint32 bOverride_CustomNearClippingPlane:1;
 
 	/** Set bOverride_CustomNearClippingPlane to true if you want to use a custom clipping plane instead of GNearClippingPlane. */
-	UPROPERTY(EditAnywhere, BlueprintReadWrite, AdvancedDisplay, Category = "Current Camera Settings", meta = (UIMin = "0.00001", ClampMin = "0.00001", editcondition = "bOverride_CustomNearClippingPlane"))
+	UPROPERTY(Interp, BlueprintSetter = SetCustomNearClippingPlane, EditAnywhere, BlueprintReadWrite, AdvancedDisplay, Category = "Current Camera Settings", meta = (UIMin = "0.00001", ClampMin = "0.00001", editcondition = "bOverride_CustomNearClippingPlane"))
 	float CustomNearClippingPlane;
-
+	
 #if WITH_EDITORONLY_DATA
 	/** Read-only. Control this value with CurrentFocalLength (and filmback settings). */
 	UPROPERTY(VisibleAnywhere, Category = "Current Camera Settings")
@@ -95,6 +95,10 @@ public:
 	UFUNCTION(BlueprintCallable, BlueprintSetter, Category = "Cine Camera")
 	void SetCurrentFocalLength(float InFocalLength);
 
+	/** Sets near clipping plane of the cine camera. */
+	UFUNCTION(BlueprintCallable, BlueprintSetter, Category = "Cine Camera")
+	void SetCustomNearClippingPlane(const float NewCustomNearClippingPlane);
+
 	/** Returns the horizonal FOV of the camera with current settings. */
 	UFUNCTION(BlueprintCallable, Category = "Cine Camera")
 	float GetHorizontalFieldOfView() const;
