CineCameraComponent: DebugFocusPlaneComponent will now be ignored when computing the component bounds for editor focus purposes. The focus plane can be very far away which makes moving the viewport to the camera itself impossible.

diff --git a/Engine/Source/Runtime/CinematicCamera/Private/CineCameraComponent.cpp b/Engine/Source/Runtime/CinematicCamera/Private/CineCameraComponent.cpp
--- a/Engine/Source/Runtime/CinematicCamera/Private/CineCameraComponent.cpp
+++ b/Engine/Source/Runtime/CinematicCamera/Private/CineCameraComponent.cpp
@@ -725,6 +725,7 @@ void UCineCameraComponent::CreateDebugFocusPlane()
 			DebugFocusPlaneComponent->CastShadow = false;
 			DebugFocusPlaneComponent->CreationMethod = CreationMethod;
 			DebugFocusPlaneComponent->bSelectable = false;
+			DebugFocusPlaneComponent->SetIgnoreBoundsForEditorFocus(true);
 
 			DebugFocusPlaneComponent->SetRelativeScale3D_Direct(FVector(10000.f, 10000.f, 1.f));
 			DebugFocusPlaneComponent->SetRelativeRotation_Direct(FRotator(90.f, 0.f, 0.f));
