Sequencer: Added custom accessors for heavily used cinematic camera properties

diff --git a/Engine/Source/Runtime/CinematicCamera/CinematicCamera.Build.cs b/Engine/Source/Runtime/CinematicCamera/CinematicCamera.Build.cs
--- a/Engine/Source/Runtime/CinematicCamera/CinematicCamera.Build.cs
+++ b/Engine/Source/Runtime/CinematicCamera/CinematicCamera.Build.cs
@@ -12,6 +12,8 @@ public class CinematicCamera : ModuleRules
 				"CoreUObject",
 				"DeveloperSettings",
                 "Engine",
+                "MovieScene",
+                "MovieSceneTracks",
                 "Slate",
                 "SlateCore"
 			}
diff --git a/Engine/Source/Runtime/CinematicCamera/Private/CinematicCameraModule.cpp b/Engine/Source/Runtime/CinematicCamera/Private/CinematicCameraModule.cpp
--- a/Engine/Source/Runtime/CinematicCamera/Private/CinematicCameraModule.cpp
+++ b/Engine/Source/Runtime/CinematicCamera/Private/CinematicCameraModule.cpp
@@ -3,4 +3,47 @@
 #include "CinematicCameraModule.h"
 #include "Modules/ModuleManager.h"
 
-IMPLEMENT_MODULE( ICinematicCameraModule, CinematicCamera );
+#include "CineCameraComponent.h"
+#include "MovieSceneTracksComponentTypes.h"
+
+
+struct FCinematicCameraModule : ICinematicCameraModule
+{
+	void StartupModule() override
+	{
+		using namespace UE::MovieScene;
+
+		FMovieSceneTracksComponentTypes* TracksComponents = FMovieSceneTracksComponentTypes::Get();
+
+		TracksComponents->Accessors.Float.Add(
+				UCineCameraComponent::StaticClass(), GET_MEMBER_NAME_CHECKED(UCineCameraComponent, CurrentAperture), 
+				GetCurrentAperture, SetCurrentAperture);
+
+		TracksComponents->Accessors.Float.Add(
+				UCineCameraComponent::StaticClass(), GET_MEMBER_NAME_CHECKED(UCineCameraComponent, CurrentFocalLength), 
+				GetCurrentFocalLength, SetCurrentFocalLength);
+	}
+
+	static float GetCurrentAperture(const UObject* Object)
+	{
+		return CastChecked<const UCineCameraComponent>(Object)->CurrentAperture;
+	}
+
+	static void SetCurrentAperture(UObject* Object, float InNewAperture)
+	{
+		CastChecked<UCineCameraComponent>(Object)->SetCurrentAperture(InNewAperture);
+	}
+
+	static float GetCurrentFocalLength(const UObject* Object)
+	{
+		return CastChecked<const UCineCameraComponent>(Object)->CurrentFocalLength;
+	}
+
+	static void SetCurrentFocalLength(UObject* Object, float InNewFocalLength)
+	{
+		CastChecked<UCineCameraComponent>(Object)->SetCurrentFocalLength(InNewFocalLength);
+	}
+};
+
+
+IMPLEMENT_MODULE(FCinematicCameraModule, CinematicCamera );
diff --git a/Engine/Source/Runtime/Engine/Private/UnrealEngine.cpp b/Engine/Source/Runtime/Engine/Private/UnrealEngine.cpp
--- a/Engine/Source/Runtime/Engine/Private/UnrealEngine.cpp
+++ b/Engine/Source/Runtime/Engine/Private/UnrealEngine.cpp
@@ -2128,6 +2128,7 @@ void UEngine::Init(IEngineLoop* InEngineLoop)
 		FModuleManager::Get().LoadModuleChecked("MovieScene");
 		FModuleManager::Get().LoadModuleChecked("MovieSceneTracks");
 		FModuleManager::Get().LoadModule("LevelSequence");
+		FModuleManager::Get().LoadModule("CinematicCamera");
 #if WITH_EDITOR
 		// The SparseVolumeTexture module containing the importer is only loaded and used in the editor.
 		FModuleManager::Get().LoadModuleChecked("SparseVolumeTexture");
