[Navigation] updated some StaticMeshComponent to never affect navigation (or set it before setting up the mesh) to eliminate useless processing

diff --git a/Engine/Source/Editor/UnrealEd/Private/ThumbnailHelpers.cpp b/Engine/Source/Editor/UnrealEd/Private/ThumbnailHelpers.cpp
--- a/Engine/Source/Editor/UnrealEd/Private/ThumbnailHelpers.cpp
+++ b/Engine/Source/Editor/UnrealEd/Private/ThumbnailHelpers.cpp
@@ -280,6 +280,7 @@ FMaterialThumbnailScene::FMaterialThumbnailScene()
 	SpawnInfo.ObjectFlags = RF_Transient;
 	PreviewActor = GetWorld()->SpawnActor<AStaticMeshActor>( SpawnInfo );
 
+	PreviewActor->GetStaticMeshComponent()->SetCanEverAffectNavigation(false);
 	PreviewActor->GetStaticMeshComponent()->SetMobility(EComponentMobility::Movable);
 	PreviewActor->GetStaticMeshComponent()->bSelectable = false;	// avoid generating hit proxies
 	PreviewActor->SetActorEnableCollision(false);
@@ -527,6 +528,7 @@ FStaticMeshThumbnailScene::FStaticMeshThumbnailScene()
 	SpawnInfo.ObjectFlags = RF_Transient;
 	PreviewActor = GetWorld()->SpawnActor<AStaticMeshActor>( SpawnInfo );
 
+	PreviewActor->GetStaticMeshComponent()->SetCanEverAffectNavigation(false);
 	PreviewActor->GetStaticMeshComponent()->SetMobility(EComponentMobility::Movable);
 	PreviewActor->SetActorEnableCollision(false);
 }
diff --git a/Engine/Source/Runtime/CinematicCamera/Private/CameraRig_Crane.cpp b/Engine/Source/Runtime/CinematicCamera/Private/CameraRig_Crane.cpp
--- a/Engine/Source/Runtime/CinematicCamera/Private/CameraRig_Crane.cpp
+++ b/Engine/Source/Runtime/CinematicCamera/Private/CameraRig_Crane.cpp
@@ -48,6 +48,7 @@ ACameraRig_Crane::ACameraRig_Crane(const FObjectInitializer& ObjectInitializer)
 		PreviewMesh_CraneBase = CreateOptionalDefaultSubobject<UStaticMeshComponent>(TEXT("PreviewMesh_CraneBase"));
 		if (PreviewMesh_CraneBase)
 		{
+			PreviewMesh_CraneBase->SetCanEverAffectNavigation(false);
 			PreviewMesh_CraneBase->SetStaticMesh(CraneBaseMesh.Object);
 			PreviewMesh_CraneBase->SetIsVisualizationComponent(true);
 			PreviewMesh_CraneBase->SetCollisionProfileName(UCollisionProfile::NoCollision_ProfileName);
@@ -61,6 +62,7 @@ ACameraRig_Crane::ACameraRig_Crane(const FObjectInitializer& ObjectInitializer)
 		PreviewMesh_CraneArm = CreateOptionalDefaultSubobject<UStaticMeshComponent>(TEXT("PreviewMesh_CraneArm"));
 		if (PreviewMesh_CraneArm)
 		{
+			PreviewMesh_CraneArm->SetCanEverAffectNavigation(false);
 			PreviewMesh_CraneArm->SetStaticMesh(CraneArmMesh.Object);
 			PreviewMesh_CraneArm->SetIsVisualizationComponent(true);
 			PreviewMesh_CraneArm->SetCollisionProfileName(UCollisionProfile::NoCollision_ProfileName);
@@ -77,6 +79,7 @@ ACameraRig_Crane::ACameraRig_Crane(const FObjectInitializer& ObjectInitializer)
 		PreviewMesh_CraneMount = CreateOptionalDefaultSubobject<UStaticMeshComponent>(TEXT("PreviewMesh_CraneMount"));
 		if (PreviewMesh_CraneMount)
 		{
+			PreviewMesh_CraneMount->SetCanEverAffectNavigation(false);
 			PreviewMesh_CraneMount->SetStaticMesh(CraneArmMount.Object);
 			PreviewMesh_CraneMount->SetIsVisualizationComponent(true);
 			PreviewMesh_CraneMount->SetCollisionProfileName(UCollisionProfile::NoCollision_ProfileName);
@@ -93,6 +96,7 @@ ACameraRig_Crane::ACameraRig_Crane(const FObjectInitializer& ObjectInitializer)
 		PreviewMesh_CraneCounterWeight = CreateOptionalDefaultSubobject<UStaticMeshComponent>(TEXT("PreviewMesh_CraneCounterWeight"));
 		if (PreviewMesh_CraneCounterWeight)
 		{
+			PreviewMesh_CraneCounterWeight->SetCanEverAffectNavigation(false);
 			PreviewMesh_CraneCounterWeight->SetStaticMesh(CraneArmCW.Object);
 			PreviewMesh_CraneCounterWeight->SetIsVisualizationComponent(true);
 			PreviewMesh_CraneCounterWeight->SetCollisionProfileName(UCollisionProfile::NoCollision_ProfileName);
diff --git a/Engine/Source/Runtime/CinematicCamera/Private/CameraRig_Rail.cpp b/Engine/Source/Runtime/CinematicCamera/Private/CameraRig_Rail.cpp
--- a/Engine/Source/Runtime/CinematicCamera/Private/CameraRig_Rail.cpp
+++ b/Engine/Source/Runtime/CinematicCamera/Private/CameraRig_Rail.cpp
@@ -49,6 +49,7 @@ ACameraRig_Rail::ACameraRig_Rail(const FObjectInitializer& ObjectInitializer)
 		PreviewMesh_Mount = CreateOptionalDefaultSubobject<UStaticMeshComponent>(TEXT("PreviewMesh_Mount"));
 		if (PreviewMesh_Mount)
 		{
+			PreviewMesh_Mount->SetCanEverAffectNavigation(false);
 			PreviewMesh_Mount->SetStaticMesh(MountMesh.Object);
 			PreviewMesh_Mount->SetIsVisualizationComponent(true);
 			PreviewMesh_Mount->SetCollisionProfileName(UCollisionProfile::NoCollision_ProfileName);
diff --git a/Engine/Source/Runtime/Engine/Private/Camera/CameraComponent.cpp b/Engine/Source/Runtime/Engine/Private/Camera/CameraComponent.cpp
--- a/Engine/Source/Runtime/Engine/Private/Camera/CameraComponent.cpp
+++ b/Engine/Source/Runtime/Engine/Private/Camera/CameraComponent.cpp
@@ -164,6 +164,7 @@ void UCameraComponent::OnRegister()
 			ProxyMeshComponent = NewObject<UCameraProxyMeshComponent>(MyOwner, NAME_None, RF_Transactional | RF_TextExportTransient);
 			ProxyMeshComponent->SetupAttachment(this);
 			ProxyMeshComponent->SetIsVisualizationComponent(true);
+			ProxyMeshComponent->SetCanEverAffectNavigation(false);
 			ProxyMeshComponent->SetStaticMesh(CameraMesh);
 			ProxyMeshComponent->SetCollisionProfileName(UCollisionProfile::NoCollision_ProfileName);
 			ProxyMeshComponent->bHiddenInGame = bCameraMeshHiddenInGame;
