Added Overscan property to UCameraComponent, which scales the camera's view frustum by scaling the projection matrix's elements to expand the view frustum by a sepcified percent. Also added film offset to the cine camera component, which can offset the projection through the off center projection matrix terms.

diff --git a/Engine/Source/Runtime/CinematicCamera/Private/CineCameraComponent.cpp b/Engine/Source/Runtime/CinematicCamera/Private/CineCameraComponent.cpp
--- a/Engine/Source/Runtime/CinematicCamera/Private/CineCameraComponent.cpp
+++ b/Engine/Source/Runtime/CinematicCamera/Private/CineCameraComponent.cpp
@@ -278,6 +278,11 @@ void UCineCameraComponent::SetCurrentFocalLength(float InFocalLength)
 }
 
 float UCineCameraComponent::GetHorizontalFieldOfView() const
+{
+	return GetHorizontalFieldOfViewInternal(/* bIncludeOverscan =*/ true);
+}
+
+float UCineCameraComponent::GetHorizontalFieldOfViewInternal(bool bIncludeOverscan) const
 {
 	if (CurrentFocalLength > 0.f)
 	{
@@ -291,13 +296,19 @@ float UCineCameraComponent::GetHorizontalFieldOfView() const
 			}
 		}
 
-		return FMath::RadiansToDegrees(2.f * FMath::Atan(CropedSensorWidth / (2.f * CurrentFocalLength)));
+		const float OverscanScalar = bIncludeOverscan ? (1.0f + Overscan) : 1.0f;
+		return FMath::RadiansToDegrees(2.f * FMath::Atan(CropedSensorWidth * OverscanScalar / (2.f * CurrentFocalLength)));
 	}
 
 	return 0.f;
 }
 
 float UCineCameraComponent::GetVerticalFieldOfView() const
+{
+	return GetVerticalFieldOfViewInternal(/* bIncludeOverscan =*/ true);
+}
+
+float UCineCameraComponent::GetVerticalFieldOfViewInternal(bool bIncludeOverscan) const
 {
 	if (CurrentFocalLength > 0.f)
 	{
@@ -311,12 +322,43 @@ float UCineCameraComponent::GetVerticalFieldOfView() const
 			}
 		}
 
-		return FMath::RadiansToDegrees(2.f * FMath::Atan(CropedSensorHeight / (2.f * CurrentFocalLength)));
+		const float OverscanScalar = bIncludeOverscan ? (1.0f + Overscan) : 1.0f;
+		return FMath::RadiansToDegrees(2.f * FMath::Atan(CropedSensorHeight * OverscanScalar / (2.f * CurrentFocalLength)));
 	}
 
 	return 0.f;
 }
 
+float UCineCameraComponent::GetHorizontalProjectionOffset() const
+{
+	float CroppedSensorWidth = Filmback.SensorWidth * LensSettings.SqueezeFactor;
+	if (CropSettings.AspectRatio > 0.0f)
+	{
+		float DesqueezeAspectRatio = Filmback.SensorWidth * LensSettings.SqueezeFactor / Filmback.SensorHeight;
+		if (CropSettings.AspectRatio < DesqueezeAspectRatio)
+		{
+			CroppedSensorWidth *= CropSettings.AspectRatio / DesqueezeAspectRatio;
+		}
+	}
+
+	return 2.0f * Filmback.SensorHorizontalOffset / (CroppedSensorWidth * (1.0f + Overscan));
+}
+
+float UCineCameraComponent::GetVerticalProjectionOffset() const
+{
+	float CroppedSensorHeight = Filmback.SensorHeight;
+	if (CropSettings.AspectRatio > 0.0f)
+	{
+		float DesqueezeAspectRatio = Filmback.SensorWidth * LensSettings.SqueezeFactor / Filmback.SensorHeight;
+		if (DesqueezeAspectRatio < CropSettings.AspectRatio)
+		{
+			CroppedSensorHeight *= DesqueezeAspectRatio / CropSettings.AspectRatio;
+		}
+	}
+
+	return 2.0f * Filmback.SensorVerticalOffset / (CroppedSensorHeight * (1.0f + Overscan));
+}
+
 FString UCineCameraComponent::GetFilmbackPresetName() const
 {
 	TArray<FNamedFilmbackPreset> const& Presets = UCineCameraSettings::GetFilmbackPresets();
@@ -484,7 +526,7 @@ void UCineCameraComponent::RecalcDerivedData()
 	float const MinFocusDistInWorldUnits = LensSettings.MinimumFocusDistance * (GetWorldToMetersScale() / 1000.f);	// convert mm to uu
 	FocusSettings.ManualFocusDistance = FMath::Max(FocusSettings.ManualFocusDistance, MinFocusDistInWorldUnits);
 
-	FieldOfView = GetHorizontalFieldOfView();
+	FieldOfView = GetHorizontalFieldOfViewInternal(/*bIncludeOverscan =*/ false);
 	Filmback.RecalcSensorAspectRatio();
 	AspectRatio = Filmback.SensorAspectRatio * LensSettings.SqueezeFactor;
 	if (CropSettings.AspectRatio > 0.0f)
@@ -557,6 +599,9 @@ void UCineCameraComponent::GetCameraView(float DeltaTime, FMinimalViewInfo& Desi
 
 	DesiredView.PerspectiveNearClipPlane = bOverride_CustomNearClippingPlane ? CustomNearClippingPlane : -1.0f;
 
+	DesiredView.OffCenterProjectionOffset.X = GetHorizontalProjectionOffset();
+	DesiredView.OffCenterProjectionOffset.Y = GetVerticalProjectionOffset();
+	
 	bResetInterpolation = false;
 }
 
@@ -699,7 +744,7 @@ void UCineCameraComponent::UpdateCameraLens(float DeltaTime, FMinimalViewInfo& D
 		DesiredView.PostProcessSettings.DepthOfFieldFocalDistance = CurrentFocusDistance;
 
 		DesiredView.PostProcessSettings.bOverride_DepthOfFieldSensorWidth = true;
-		DesiredView.PostProcessSettings.DepthOfFieldSensorWidth = Filmback.SensorWidth;
+		DesiredView.PostProcessSettings.DepthOfFieldSensorWidth = Filmback.SensorWidth * (1.0f + Overscan);
 
 		DesiredView.PostProcessSettings.bOverride_DepthOfFieldSqueezeFactor = true;
 		DesiredView.PostProcessSettings.DepthOfFieldSqueezeFactor = LensSettings.SqueezeFactor;
diff --git a/Engine/Source/Runtime/CinematicCamera/Public/CineCameraComponent.h b/Engine/Source/Runtime/CinematicCamera/Public/CineCameraComponent.h
--- a/Engine/Source/Runtime/CinematicCamera/Public/CineCameraComponent.h
+++ b/Engine/Source/Runtime/CinematicCamera/Public/CineCameraComponent.h
@@ -107,6 +107,12 @@ public:
 	UFUNCTION(BlueprintCallable, Category = "Cine Camera")
 	CINEMATICCAMERA_API float GetVerticalFieldOfView() const;
 
+	UFUNCTION(BlueprintCallable, BlueprintPure, Category = "Cine Camera")
+	CINEMATICCAMERA_API float GetHorizontalProjectionOffset() const;
+
+	UFUNCTION(BlueprintCallable, BlueprintPure, Category = "Cine Camera")
+	CINEMATICCAMERA_API float GetVerticalProjectionOffset() const;
+	
 	/** Returns the filmback name of the camera with the current settings. */
 	UFUNCTION(BlueprintCallable, Category = "Cine Camera")
 	CINEMATICCAMERA_API FString GetFilmbackPresetName() const;
@@ -252,7 +258,9 @@ private:
 	void SetLensPresetByNameInternal(const FString& InPresetName);
 	void SetFilmbackPresetByNameInternal(const FString& InPresetName, FCameraFilmbackSettings& InOutFilmbackSettings);
 	void SetCropPresetByNameInternal(const FString& InPresetName);
-
+	float GetHorizontalFieldOfViewInternal(bool bIncludeOverscan = true) const;
+	float GetVerticalFieldOfViewInternal(bool bIncludeOverscan) const;
+	
 #if WITH_EDITORONLY_DATA
 	void CreateDebugFocusPlane();
 	void DestroyDebugFocusPlane();
diff --git a/Engine/Source/Runtime/CinematicCamera/Public/CineCameraSettings.h b/Engine/Source/Runtime/CinematicCamera/Public/CineCameraSettings.h
--- a/Engine/Source/Runtime/CinematicCamera/Public/CineCameraSettings.h
+++ b/Engine/Source/Runtime/CinematicCamera/Public/CineCameraSettings.h
@@ -24,6 +24,14 @@ struct FCameraFilmbackSettings
 	UPROPERTY(Interp, EditAnywhere, BlueprintReadWrite, Category = "Filmback", meta = (ClampMin = "0.001", ForceUnits = mm))
 	float SensorHeight;
 
+	/** Horizontal offset of the filmback, in mm. */
+	UPROPERTY(Interp, EditAnywhere, BlueprintReadWrite, Category = "Filmback", meta = (ForceUnits = mm))
+	float SensorHorizontalOffset;
+
+	/** Vertical offset of the filmback, in mm. */
+	UPROPERTY(Interp, EditAnywhere, BlueprintReadWrite, Category = "Filmback", meta = (ForceUnits = mm))
+	float SensorVerticalOffset;
+	
 	/** Read-only. Computed from Sensor dimensions. */
 	UPROPERTY(Interp, VisibleAnywhere, BlueprintReadOnly, Category = "Filmback")
 	float SensorAspectRatio;
@@ -31,7 +39,9 @@ struct FCameraFilmbackSettings
 	bool operator==(const FCameraFilmbackSettings& Other) const
 	{
 		return (SensorWidth == Other.SensorWidth)
-			&& (SensorHeight == Other.SensorHeight);
+			&& (SensorHeight == Other.SensorHeight)
+			&& (SensorHorizontalOffset == Other.SensorHorizontalOffset)
+			&& (SensorVerticalOffset == Other.SensorVerticalOffset);
 	}
 
 	bool operator!=(const FCameraFilmbackSettings& Other) const
@@ -47,6 +57,8 @@ struct FCameraFilmbackSettings
 	FCameraFilmbackSettings()
 		: SensorWidth(24.89f)
 		, SensorHeight(18.67f)
+		, SensorHorizontalOffset(0.0f)
+		, SensorVerticalOffset(0.0f)
 		, SensorAspectRatio(1.33f)
 	{
 	}
diff --git a/Engine/Source/Runtime/Engine/Classes/Camera/CameraComponent.h b/Engine/Source/Runtime/Engine/Classes/Camera/CameraComponent.h
--- a/Engine/Source/Runtime/Engine/Classes/Camera/CameraComponent.h
+++ b/Engine/Source/Runtime/Engine/Classes/Camera/CameraComponent.h
@@ -130,7 +130,12 @@ class UCameraComponent : public USceneComponent
 	uint8 bUseFieldOfViewForLOD : 1;
 	UFUNCTION(BlueprintCallable, Category = Camera)
 	void SetUseFieldOfViewForLOD(bool bInUseFieldOfViewForLOD) { bUseFieldOfViewForLOD = bInUseFieldOfViewForLOD; }
-
+	/** Amount to increase the view frustum by, from 0.0 for no increase to 1.0 for 100% increase */
+	UPROPERTY(Interp, EditAnywhere, BlueprintReadWrite, Category = CameraOptions, meta = (UIMin="0.0", ClampMin="0.0", UIMax="1.0", ClampMax="1.0"))
+	float Overscan;
+	UFUNCTION(BlueprintCallable, Category = Camera)
+	void SetOverscan(float InOverscan) { Overscan = InOverscan; }
+	
 #if WITH_EDITOR
 	// Returns the filmback text used for burnins on preview viewports
 	UFUNCTION()
diff --git a/Engine/Source/Runtime/Engine/Classes/Camera/CameraTypes.h b/Engine/Source/Runtime/Engine/Classes/Camera/CameraTypes.h
--- a/Engine/Source/Runtime/Engine/Classes/Camera/CameraTypes.h
+++ b/Engine/Source/Runtime/Engine/Classes/Camera/CameraTypes.h
@@ -214,4 +214,9 @@ public:
 	 * It is computed as tan(SceneFOVRadians * 0.5) / tan(FirstPersonFOVRadians * 0.5).
 	 */
 	ENGINE_API float CalculateFirstPersonFOVCorrectionFactor() const;
+
+	/**
+	 * Apply overscan to the view info, which scales the field of view and ortho width to simulate expanding the view frustum.
+	 */
+	ENGINE_API void ApplyOverscan(float InOverscan);
 };
diff --git a/Engine/Source/Runtime/Engine/Classes/Components/SceneCaptureComponent2D.h b/Engine/Source/Runtime/Engine/Classes/Components/SceneCaptureComponent2D.h
--- a/Engine/Source/Runtime/Engine/Classes/Components/SceneCaptureComponent2D.h
+++ b/Engine/Source/Runtime/Engine/Classes/Components/SceneCaptureComponent2D.h
@@ -50,6 +50,10 @@ public:
 	UPROPERTY(EditAnywhere, BlueprintReadWrite, Category=Projection, meta = (editcondition = "ProjectionType==1 && bUpdateOrthoPlanes"))
 	bool bUseCameraHeightAsViewTarget;
 
+	/** Amount to increase the view frustum by, from 0.0 for no increase to 1.0 for 100% increase */
+	UPROPERTY(EditAnywhere, BlueprintReadWrite, Category = Projection, meta = (UIMin="0.0", ClampMin="0.0", UIMax="1.0", ClampMax="1.0"))
+	float Overscan;
+	
 	/** Output render target of the scene capture that can be read in materials. */
 	UPROPERTY(EditAnywhere, BlueprintReadWrite, Category=SceneCapture)
 	TObjectPtr<class UTextureRenderTarget2D> TextureTarget;
diff --git a/Engine/Source/Runtime/Engine/Private/Camera/CameraComponent.cpp b/Engine/Source/Runtime/Engine/Private/Camera/CameraComponent.cpp
--- a/Engine/Source/Runtime/Engine/Private/Camera/CameraComponent.cpp
+++ b/Engine/Source/Runtime/Engine/Private/Camera/CameraComponent.cpp
@@ -85,6 +85,7 @@ UCameraComponent::UCameraComponent(const FObjectInitializer& ObjectInitializer)
 	OrthoFarClipPlane = DEFAULT_ORTHOFARPLANE;
 	bConstrainAspectRatio = false;
 	bOverrideAspectRatioAxisConstraint = false;
+	Overscan = 0.0f;
 	bUseFieldOfViewForLOD = true;
 	PostProcessBlendWeight = 1.0f;
 	bUsePawnControlRotation = false;
@@ -463,6 +464,8 @@ void UCameraComponent::GetCameraView(float DeltaTime, FMinimalViewInfo& DesiredV
 	DesiredView.FirstPersonScale = bEnableFirstPersonScale ? FirstPersonScale : 1.0f;
 	DesiredView.bUseFirstPersonParameters = bEnableFirstPersonFieldOfView || bEnableFirstPersonScale;
 	
+	DesiredView.ApplyOverscan(Overscan);
+	
 	if (bAutoCalculateOrthoPlanes)
 	{
 		if (const AActor* ViewTarget = GetOwner())
diff --git a/Engine/Source/Runtime/Engine/Private/Camera/CameraStackTypes.cpp b/Engine/Source/Runtime/Engine/Private/Camera/CameraStackTypes.cpp
--- a/Engine/Source/Runtime/Engine/Private/Camera/CameraStackTypes.cpp
+++ b/Engine/Source/Runtime/Engine/Private/Camera/CameraStackTypes.cpp
@@ -469,3 +469,17 @@ float FMinimalViewInfo::CalculateFirstPersonFOVCorrectionFactor() const
 	const float FOVCorrectionFactor = HalfTanSceneFOV / HalfTanFirstPersonFOV;
 	return FOVCorrectionFactor;
 }
+
+void FMinimalViewInfo::ApplyOverscan(float InOverscan)
+{
+	if (InOverscan > 0.0)
+	{
+		// Overscan directly scales the view frustum, but can be accomplished by scaling the FOV.
+		// However, must scale the tangent of the half-FOV to accomplish the same mathematical transform.
+		const float HalfFOVInRadians = FMath::DegreesToRadians(0.5f * FOV);
+		const float OverscannedFOV = FMath::Atan((1.0f + InOverscan) * FMath::Tan(HalfFOVInRadians));
+		FOV = 2.0f * FMath::RadiansToDegrees(OverscannedFOV);
+		
+		OrthoWidth *= 1.0f + InOverscan;
+	}
+}
diff --git a/Engine/Source/Runtime/Engine/Private/Components/SceneCaptureComponent.cpp b/Engine/Source/Runtime/Engine/Private/Components/SceneCaptureComponent.cpp
--- a/Engine/Source/Runtime/Engine/Private/Components/SceneCaptureComponent.cpp
+++ b/Engine/Source/Runtime/Engine/Private/Components/SceneCaptureComponent.cpp
@@ -595,7 +595,8 @@ USceneCaptureComponent2D::USceneCaptureComponent2D(const FObjectInitializer& Obj
 	AutoPlaneShift = 0.0f;
 	bUpdateOrthoPlanes = false;
 	bUseCameraHeightAsViewTarget = false;
-
+	Overscan = 0.0;
+	
 	bUseCustomProjectionMatrix = false;
 	bAutoActivate = true;
 	PrimaryComponentTick.bCanEverTick = true;
@@ -732,6 +733,8 @@ void USceneCaptureComponent2D::GetCameraView(float DeltaTime, FMinimalViewInfo&
 	OutMinimalViewInfo.bUpdateOrthoPlanes = bUpdateOrthoPlanes;
 	OutMinimalViewInfo.bUseCameraHeightAsViewTarget = bUseCameraHeightAsViewTarget;
 
+	OutMinimalViewInfo.ApplyOverscan(Overscan);
+	
 	if (bAutoCalculateOrthoPlanes)
 	{
 		if(const AActor* ViewTarget = GetOwner())
diff --git a/Engine/Source/Runtime/Renderer/Private/SceneCaptureRendering.cpp b/Engine/Source/Runtime/Renderer/Private/SceneCaptureRendering.cpp
--- a/Engine/Source/Runtime/Renderer/Private/SceneCaptureRendering.cpp
+++ b/Engine/Source/Runtime/Renderer/Private/SceneCaptureRendering.cpp
@@ -1054,7 +1054,8 @@ void FScene::UpdateSceneCaptureContents(USceneCaptureComponent2D* CaptureCompone
 			FPlane(1, 0, 0, 0),
 			FPlane(0, 1, 0, 0),
 			FPlane(0, 0, 0, 1));
-		const float FOV = CaptureComponent->FOVAngle * (float)PI / 360.0f;
+		const float UnscaledFOV = CaptureComponent->FOVAngle * (float)PI / 360.0f;
+		const float FOV = FMath::Atan((1.0f + CaptureComponent->Overscan) * FMath::Tan(UnscaledFOV));
 		FIntPoint CaptureSize(TextureRenderTarget->GetSurfaceWidth(), TextureRenderTarget->GetSurfaceHeight());
 
 		const bool bUseSceneColorTexture = CaptureNeedsSceneColor(CaptureComponent->CaptureSource);
