CineCamera: We always override the f-stop no matter the focus method, since aperture also impacts physical camera exposure.

We also remove redundant tests against the fstop property. Focal distance already acts as the primary state to check whether DoF is active.

diff --git a/Engine/Plugins/Runtime/nDisplay/Source/DisplayCluster/Private/Render/Viewport/DisplayClusterViewport_CustomPostProcessSettings.cpp b/Engine/Plugins/Runtime/nDisplay/Source/DisplayCluster/Private/Render/Viewport/DisplayClusterViewport_CustomPostProcessSettings.cpp
--- a/Engine/Plugins/Runtime/nDisplay/Source/DisplayCluster/Private/Render/Viewport/DisplayClusterViewport_CustomPostProcessSettings.cpp
+++ b/Engine/Plugins/Runtime/nDisplay/Source/DisplayCluster/Private/Render/Viewport/DisplayClusterViewport_CustomPostProcessSettings.cpp
@@ -46,8 +46,7 @@ namespace UE::DisplayClusterViewport::CustomPostProcess
 		check(InViewport);
 
 		// This math only works for the CineCamera DoF because it provides a valid 'DepthOfFieldSensorFocalLength' value.
-		if (!(InOutPostProcessSettings.DepthOfFieldFstop > 0.0f
-			&& InOutPostProcessSettings.DepthOfFieldFocalDistance > 0.0f))
+		if (InOutPostProcessSettings.DepthOfFieldFocalDistance <= 0.0f)
 		{
 			return false;
 		}
diff --git a/Engine/Plugins/Runtime/nDisplay/Source/DisplayCluster/Private/Render/Viewport/DisplayClusterViewport_ViewPoint.cpp b/Engine/Plugins/Runtime/nDisplay/Source/DisplayCluster/Private/Render/Viewport/DisplayClusterViewport_ViewPoint.cpp
--- a/Engine/Plugins/Runtime/nDisplay/Source/DisplayCluster/Private/Render/Viewport/DisplayClusterViewport_ViewPoint.cpp
+++ b/Engine/Plugins/Runtime/nDisplay/Source/DisplayCluster/Private/Render/Viewport/DisplayClusterViewport_ViewPoint.cpp
@@ -185,7 +185,7 @@ bool FDisplayClusterViewport::SetupViewPoint(const uint32 InContextNum, FMinimal
 			FDisplayClusterViewport_Context& DestContext = Contexts[InContextNum];
 
 			// DoF FocalLength
-			if (InOutViewInfo.PostProcessSettings.DepthOfFieldFstop > 0.0f && InOutViewInfo.PostProcessSettings.DepthOfFieldFocalDistance > 0.0f)
+			if (InOutViewInfo.PostProcessSettings.DepthOfFieldFocalDistance > 0.0f)
 			{
 				// Convert FOV to focal length,
 				// 
diff --git a/Engine/Source/Runtime/CinematicCamera/Private/CineCameraComponent.cpp b/Engine/Source/Runtime/CinematicCamera/Private/CineCameraComponent.cpp
--- a/Engine/Source/Runtime/CinematicCamera/Private/CineCameraComponent.cpp
+++ b/Engine/Source/Runtime/CinematicCamera/Private/CineCameraComponent.cpp
@@ -695,9 +695,12 @@ void UCineCameraComponent::UpdateDebugFocusPlane()
 
 void UCineCameraComponent::UpdateCameraLens(float DeltaTime, FMinimalViewInfo& DesiredView)
 {
+	// UE-192569: Given that aperture impacts exposure, we always override it.
+	DesiredView.PostProcessSettings.bOverride_DepthOfFieldFstop = true;
+	DesiredView.PostProcessSettings.DepthOfFieldFstop = CurrentAperture;
+
 	if (FocusSettings.FocusMethod == ECameraFocusMethod::DoNotOverride)
 	{
-		DesiredView.PostProcessSettings.bOverride_DepthOfFieldFstop = false;
 		DesiredView.PostProcessSettings.bOverride_DepthOfFieldMinFstop = false;
 		DesiredView.PostProcessSettings.bOverride_DepthOfFieldBladeCount = false;
 		DesiredView.PostProcessSettings.bOverride_DepthOfFieldFocalDistance = false;
@@ -715,9 +718,6 @@ void UCineCameraComponent::UpdateCameraLens(float DeltaTime, FMinimalViewInfo& D
 		// Update focus/DoF
 		DesiredView.PostProcessBlendWeight = PostProcessBlendWeight;
 
-		DesiredView.PostProcessSettings.bOverride_DepthOfFieldFstop = true;
-		DesiredView.PostProcessSettings.DepthOfFieldFstop = CurrentAperture;
-
 		DesiredView.PostProcessSettings.bOverride_DepthOfFieldMinFstop = true;
 		DesiredView.PostProcessSettings.DepthOfFieldMinFstop = LensSettings.MinFStop;
 
diff --git a/Engine/Source/Runtime/Renderer/Private/PathTracing.cpp b/Engine/Source/Runtime/Renderer/Private/PathTracing.cpp
--- a/Engine/Source/Runtime/Renderer/Private/PathTracing.cpp
+++ b/Engine/Source/Runtime/Renderer/Private/PathTracing.cpp
@@ -825,8 +825,7 @@ static void PreparePathTracingData(const FScene* Scene, const FViewInfo& View, F
 	PathTracingData.CameraLensRadius = FVector2f::ZeroVector;
 	if (ShowFlags.DepthOfField &&
 		PathTracing::UsesReferenceDOF(View) &&
-		PPV.DepthOfFieldFocalDistance > 0 &&
-		PPV.DepthOfFieldFstop > 0)
+		PPV.DepthOfFieldFocalDistance > 0)
 	{
 		DiaphragmDOF::FPhysicalCocModel CocModel;
 		CocModel.Compile(View);
diff --git a/Engine/Source/Runtime/Renderer/Private/PostProcess/DiaphragmDOF.cpp b/Engine/Source/Runtime/Renderer/Private/PostProcess/DiaphragmDOF.cpp
--- a/Engine/Source/Runtime/Renderer/Private/PostProcess/DiaphragmDOF.cpp
+++ b/Engine/Source/Runtime/Renderer/Private/PostProcess/DiaphragmDOF.cpp
@@ -1376,7 +1376,7 @@ bool DiaphragmDOF::IsEnabled(const FViewInfo& View)
 		View.Family->EngineShowFlags.DepthOfField &&
 		bDepthOfFieldRequestedByCVar &&
 		!(View.Family->EngineShowFlags.PathTracing && PathTracing::UsesReferenceDOF(View)) &&
-		((View.FinalPostProcessSettings.DepthOfFieldFstop > 0.f && View.FinalPostProcessSettings.DepthOfFieldFocalDistance > 0.f) || View.FinalPostProcessSettings.DepthOfFieldDepthBlurRadius > 0.f);
+		(View.FinalPostProcessSettings.DepthOfFieldFocalDistance > 0.f || View.FinalPostProcessSettings.DepthOfFieldDepthBlurRadius > 0.f);
 }
 
 bool DiaphragmDOF::AddPasses(
diff --git a/Engine/Source/Runtime/Renderer/Private/PostProcess/DiaphragmDOFUtils.cpp b/Engine/Source/Runtime/Renderer/Private/PostProcess/DiaphragmDOFUtils.cpp
--- a/Engine/Source/Runtime/Renderer/Private/PostProcess/DiaphragmDOFUtils.cpp
+++ b/Engine/Source/Runtime/Renderer/Private/PostProcess/DiaphragmDOFUtils.cpp
@@ -49,7 +49,7 @@ float ComputeFocalLengthFromFov(const FSceneView& View)
 FVector4f DiaphragmDOF::CircleDofHalfCoc(const FViewInfo& View)
 {
 	static const auto CVar = IConsoleManager::Get().FindTConsoleVariableDataInt(TEXT("r.DepthOfFieldQuality"));
-	bool bDepthOfField = View.Family->EngineShowFlags.DepthOfField && CVar->GetValueOnRenderThread() > 0 && View.FinalPostProcessSettings.DepthOfFieldFstop > 0 && View.FinalPostProcessSettings.DepthOfFieldFocalDistance > 0;
+	bool bDepthOfField = View.Family->EngineShowFlags.DepthOfField && CVar->GetValueOnRenderThread() > 0 && View.FinalPostProcessSettings.DepthOfFieldFocalDistance > 0;
 
 	FVector4f Ret(0, 1, 0, 0);
 
