Camera: Follow up to 38856257 use the pre-existing name if the DisplayName is not set

#rnx

diff --git a/Engine/Source/Runtime/CinematicCamera/Private/CineCameraComponent.cpp b/Engine/Source/Runtime/CinematicCamera/Private/CineCameraComponent.cpp
--- a/Engine/Source/Runtime/CinematicCamera/Private/CineCameraComponent.cpp
+++ b/Engine/Source/Runtime/CinematicCamera/Private/CineCameraComponent.cpp
@@ -620,7 +620,7 @@ FText UCineCameraComponent::GetFilmbackText() const
 	{
 		return FText::Format(
 			LOCTEXT("PresetFormat","FilmbackPreset: {0} | Zoom: {1}mm | Av: {2} | Squeeze: {3}"),
-			Preset->DisplayName,
+			Preset->DisplayName.IsEmpty() ? FText::FromString(Preset->Name) : Preset->DisplayName,
 			FText::AsNumber(CurrentFocalLength),
 			FText::AsNumber(CurrentAperture),
 			FText::AsNumber(LensSettings.SqueezeFactor)
