Cine Camera: Fix camera rig rail display when Closed Loop is enabled (the last loop wasn't displaying the segment)

diff --git a/Engine/Source/Runtime/CinematicCamera/Private/CameraRig_Rail.cpp b/Engine/Source/Runtime/CinematicCamera/Private/CameraRig_Rail.cpp
--- a/Engine/Source/Runtime/CinematicCamera/Private/CameraRig_Rail.cpp
+++ b/Engine/Source/Runtime/CinematicCamera/Private/CameraRig_Rail.cpp
@@ -88,8 +88,7 @@ void ACameraRig_Rail::UpdatePreviewMeshes()
 	{
 		if (PreviewRailStaticMesh)
 		{
-			int32 const NumSplinePoints = RailSplineComponent->GetNumberOfSplinePoints();
-			int32 const NumNeededPreviewMeshes = FMath::Max(0, NumSplinePoints - 1);
+			int32 const NumNeededPreviewMeshes = FMath::Max(0, RailSplineComponent->GetNumberOfSplineSegments());
 
 			// make sure our preview mesh array is correctly sized and populated
 			{
@@ -119,7 +118,7 @@ void ACameraRig_Rail::UpdatePreviewMeshes()
 				check(PreviewRailMeshSegments.Num() == NumNeededPreviewMeshes);
 			}
 
-			for (int PtIdx = 0; PtIdx < NumSplinePoints - 1; ++PtIdx)
+			for (int PtIdx = 0; PtIdx < NumNeededPreviewMeshes; ++PtIdx)
 			{
 				FVector StartLoc = RailSplineComponent->GetLocationAtSplinePoint(PtIdx, ESplineCoordinateSpace::Local);
 				FVector StartTangent = RailSplineComponent->GetLeaveTangentAtSplinePoint(PtIdx, ESplineCoordinateSpace::Local);
