CineCamera: add option to disable dynamic exposure

Previously, turning off depth-of-field with FocusMethod would also turn off exposure simulation. This was fixed recently, but also meant that CineCameraComponent would now always simulate exposure, which isn't always desirable. This change adds a new ExposureMethod setting to allow disabling exposure too.

diff --git a/Engine/Source/Runtime/CinematicCamera/Private/CineCameraComponent.cpp b/Engine/Source/Runtime/CinematicCamera/Private/CineCameraComponent.cpp
--- a/Engine/Source/Runtime/CinematicCamera/Private/CineCameraComponent.cpp
+++ b/Engine/Source/Runtime/CinematicCamera/Private/CineCameraComponent.cpp
@@ -56,6 +56,7 @@ UCineCameraComponent::UCineCameraComponent()
 	CurrentAperture = CineCameraSettings->DefaultLensFStop;
 	CurrentFocalLength = CineCameraSettings->DefaultLensFocalLength;
 
+	ExposureMethod = ECameraExposureMethod::Enabled;
 
 	RecalcDerivedData();
 
@@ -699,9 +700,18 @@ void UCineCameraComponent::UpdateDebugFocusPlane()
 
 void UCineCameraComponent::UpdateCameraLens(float DeltaTime, FMinimalViewInfo& DesiredView)
 {
-	// UE-192569: Given that aperture impacts exposure, we always override it.
-	DesiredView.PostProcessSettings.bOverride_DepthOfFieldFstop = true;
-	DesiredView.PostProcessSettings.DepthOfFieldFstop = CurrentAperture;
+	if (ExposureMethod == ECameraExposureMethod::DoNotOverride &&
+			FocusSettings.FocusMethod == ECameraFocusMethod::DoNotOverride)
+	{
+		// We want neither dynamic exposure nor focus effects.
+		DesiredView.PostProcessSettings.bOverride_DepthOfFieldFstop = false;
+	}
+	else
+	{
+		// We want dynamic exposure or focus effects, maybe both.
+		DesiredView.PostProcessSettings.bOverride_DepthOfFieldFstop = true;
+		DesiredView.PostProcessSettings.DepthOfFieldFstop = CurrentAperture;
+	}
 
 	if (FocusSettings.FocusMethod == ECameraFocusMethod::DoNotOverride)
 	{
diff --git a/Engine/Source/Runtime/CinematicCamera/Public/CineCameraComponent.h b/Engine/Source/Runtime/CinematicCamera/Public/CineCameraComponent.h
--- a/Engine/Source/Runtime/CinematicCamera/Public/CineCameraComponent.h
+++ b/Engine/Source/Runtime/CinematicCamera/Public/CineCameraComponent.h
@@ -76,6 +76,10 @@ public:
 	UPROPERTY(VisibleAnywhere, BlueprintReadOnly, Category = "Current Camera Settings")
 	float CurrentFocusDistance;
 
+	/** Exposure method, for determining whether exposure-related post-process settings should be automatically set from camera parameters */
+	UPROPERTY(Interp, EditAnywhere, BlueprintReadWrite, Category = "Current Camera Settings")
+	ECameraExposureMethod ExposureMethod;
+
 	UPROPERTY(EditAnywhere, BlueprintReadWrite, AdvancedDisplay, Category = "Current Camera Settings", meta = (InlineEditConditionToggle))
 	uint32 bOverride_CustomNearClippingPlane:1;
 
diff --git a/Engine/Source/Runtime/CinematicCamera/Public/CineCameraSettings.h b/Engine/Source/Runtime/CinematicCamera/Public/CineCameraSettings.h
--- a/Engine/Source/Runtime/CinematicCamera/Public/CineCameraSettings.h
+++ b/Engine/Source/Runtime/CinematicCamera/Public/CineCameraSettings.h
@@ -283,6 +283,19 @@ struct FCameraFocusSettings
 	{}
 };
 
+/** Supported methods for exposure handling. */
+UENUM()
+enum class ECameraExposureMethod : uint8
+{
+	/** Don't override, ie. allow post process volume settings to persist. */
+	DoNotOverride,
+
+	/** Set exposure settings based on the current lens and filmback settings. */
+	Enabled,
+
+	MAX UMETA(Hidden)
+};
+
 UCLASS(config=Engine, defaultconfig, meta=(DisplayName="Cinematic Camera"), MinimalAPI)
 class UCineCameraSettings : public UDeveloperSettings
 {
