[CinematicCamera] Fix for camera lookats getting clobbered by re-flushing of all camera tracks whenever a second movie sequence is spawned during a frame with a level sequence with camera lookats.

diff --git a/Engine/Source/Runtime/CinematicCamera/Private/CineCameraActor.cpp b/Engine/Source/Runtime/CinematicCamera/Private/CineCameraActor.cpp
--- a/Engine/Source/Runtime/CinematicCamera/Private/CineCameraActor.cpp
+++ b/Engine/Source/Runtime/CinematicCamera/Private/CineCameraActor.cpp
@@ -83,7 +83,11 @@ void ACineCameraActor::Tick(float DeltaTime)
 
 			// we store this ourselves in case other systems try to change our rotation, and end fighting the interpolation
 			LookatTrackingSettings.LastLookatTrackingRotation = FinalRot;
-
+			// the frame as well, so the CameraComponent can fight those other systems.
+			// This is a bug fix for when starting a new level sequence forces a re-flush of all timelines
+			// including camera transforms.  That's probably the only other system referenced above.
+			LookatTrackingSettings.LastLookatTrackingRotationFrame = GFrameNumber;
+			
 #if ENABLE_DRAW_DEBUG
 			if (LookatTrackingSettings.bDrawDebugLookAtTrackingPosition)
 			{
diff --git a/Engine/Source/Runtime/CinematicCamera/Public/CineCameraActor.h b/Engine/Source/Runtime/CinematicCamera/Public/CineCameraActor.h
--- a/Engine/Source/Runtime/CinematicCamera/Public/CineCameraActor.h
+++ b/Engine/Source/Runtime/CinematicCamera/Public/CineCameraActor.h
@@ -29,6 +29,8 @@ struct FCameraLookatTrackingSettings
 		, bDrawDebugLookAtTrackingPosition(false)
 		, LookAtTrackingInterpSpeed(0.f)
 		, LastLookatTrackingRotation(FRotator::ZeroRotator)
+		, LastLookatTrackingRotationFrame(-1)
+		, bFixupLastLookatTrackingRotation(false)
 		, RelativeOffset(FVector::ZeroVector)
 		, bAllowRoll(false)
 	{
@@ -48,7 +50,13 @@ struct FCameraLookatTrackingSettings
 
 	/** Last known lookat tracking rotation (used during interpolation) */
 	FRotator LastLookatTrackingRotation;
+	
+	/** The frame we set the above on, so the camera component can restore it */
+	uint32 LastLookatTrackingRotationFrame;
 
+	/* If true, tells the camera component to fixup the lookat rotation.  Last minute fix, see the code for details. */
+	uint8 bFixupLastLookatTrackingRotation : 1;
+	
 	/** If set, camera will track this actor's location */
 	UPROPERTY(Interp, EditAnywhere, BlueprintReadWrite, Category = "LookAt")
 	TSoftObjectPtr<AActor> ActorToTrack;
diff --git a/Engine/Source/Runtime/Engine/Private/Camera/CameraComponent.cpp b/Engine/Source/Runtime/Engine/Private/Camera/CameraComponent.cpp
--- a/Engine/Source/Runtime/Engine/Private/Camera/CameraComponent.cpp
+++ b/Engine/Source/Runtime/Engine/Private/Camera/CameraComponent.cpp
@@ -1,7 +1,9 @@
 // Copyright Epic Games, Inc. All Rights Reserved.
 
 #include "Camera/CameraComponent.h"
+
 #include "Camera/CameraTypes.h"
+#include "CineCameraActor.h"
 #include "SceneView.h"
 #include "UObject/ConstructorHelpers.h"
 #include "Engine/World.h"
@@ -433,6 +435,15 @@ void UCameraComponent::GetCameraView(float DeltaTime, FMinimalViewInfo& DesiredV
 			}
 		}
 	}
+	else if (TObjectPtr<ACineCameraActor> Cam = Cast<ACineCameraActor>(GetOwner()))
+	{
+		// TODO maybe we should do the lookat logic here instead?  Maybe with some blend amount if that can be determined.
+		// Or defer the Camera tick group until the latest it can be.
+		if (Cam->LookatTrackingSettings.bFixupLastLookatTrackingRotation && Cam->LookatTrackingSettings.LastLookatTrackingRotationFrame == GFrameNumber)
+		{
+			SetWorldRotation(Cam->LookatTrackingSettings.LastLookatTrackingRotation);
+		}
+	}
 
 	if (bUseAdditiveOffset)
 	{
