[Sparks] Removing the safety gating lookat fixup member from CineCameraActor.h.  It's too constricting, since we don't even use SparksLevelSequencePlayers for some of the battlestage intro shots.

#tests PIE Battlestage

diff --git a/Engine/Source/Runtime/CinematicCamera/Public/CineCameraActor.h b/Engine/Source/Runtime/CinematicCamera/Public/CineCameraActor.h
--- a/Engine/Source/Runtime/CinematicCamera/Public/CineCameraActor.h
+++ b/Engine/Source/Runtime/CinematicCamera/Public/CineCameraActor.h
@@ -30,7 +30,6 @@ struct FCameraLookatTrackingSettings
 		, LookAtTrackingInterpSpeed(0.f)
 		, LastLookatTrackingRotation(FRotator::ZeroRotator)
 		, LastLookatTrackingRotationFrame(-1)
-		, bFixupLastLookatTrackingRotation(false)
 		, RelativeOffset(FVector::ZeroVector)
 		, bAllowRoll(false)
 	{
@@ -54,9 +53,6 @@ struct FCameraLookatTrackingSettings
 	/** The frame we set the above on, so the camera component can restore it */
 	uint32 LastLookatTrackingRotationFrame;
 
-	/* If true, tells the camera component to fixup the lookat rotation.  Last minute fix, see the code for details. */
-	uint8 bFixupLastLookatTrackingRotation : 1;
-	
 	/** If set, camera will track this actor's location */
 	UPROPERTY(Interp, EditAnywhere, BlueprintReadWrite, Category = "LookAt")
 	TSoftObjectPtr<AActor> ActorToTrack;
diff --git a/Engine/Source/Runtime/Engine/Private/Camera/CameraComponent.cpp b/Engine/Source/Runtime/Engine/Private/Camera/CameraComponent.cpp
--- a/Engine/Source/Runtime/Engine/Private/Camera/CameraComponent.cpp
+++ b/Engine/Source/Runtime/Engine/Private/Camera/CameraComponent.cpp
@@ -439,7 +439,7 @@ void UCameraComponent::GetCameraView(float DeltaTime, FMinimalViewInfo& DesiredV
 	{
 		// TODO maybe we should do the lookat logic here instead?  Maybe with some blend amount if that can be determined.
 		// Or defer the Camera tick group until the latest it can be.
-		if (Cam->LookatTrackingSettings.bFixupLastLookatTrackingRotation && Cam->LookatTrackingSettings.LastLookatTrackingRotationFrame == GFrameNumber)
+		if (Cam->LookatTrackingSettings.LastLookatTrackingRotationFrame == GFrameNumber)
 		{
 			SetWorldRotation(Cam->LookatTrackingSettings.LastLookatTrackingRotation);
 		}
