Cine Camera: Add Units = cm for Custom Near Clipping Plane

#FailedValidations Preflight

diff --git a/Engine/Source/Runtime/CinematicCamera/Public/CineCameraComponent.h b/Engine/Source/Runtime/CinematicCamera/Public/CineCameraComponent.h
--- a/Engine/Source/Runtime/CinematicCamera/Public/CineCameraComponent.h
+++ b/Engine/Source/Runtime/CinematicCamera/Public/CineCameraComponent.h
@@ -84,7 +84,7 @@ public:
 	uint32 bOverride_CustomNearClippingPlane:1;
 
 	/** Set bOverride_CustomNearClippingPlane to true if you want to use a custom clipping plane instead of GNearClippingPlane. */
-	UPROPERTY(Interp, BlueprintSetter = SetCustomNearClippingPlane, EditAnywhere, BlueprintReadWrite, AdvancedDisplay, Category = "Current Camera Settings", meta = (UIMin = "0.00001", ClampMin = "0.00001", editcondition = "bOverride_CustomNearClippingPlane"))
+	UPROPERTY(Interp, BlueprintSetter = SetCustomNearClippingPlane, EditAnywhere, BlueprintReadWrite, AdvancedDisplay, Category = "Current Camera Settings", meta = (UIMin = "0.00001", ClampMin = "0.00001", editcondition = "bOverride_CustomNearClippingPlane", Units = cm))
 	float CustomNearClippingPlane;
 	
 #if WITH_EDITORONLY_DATA
