Cine Camera: Remove Interp from SensorAspectRatio which is a value derived from Sensor Width/Height

diff --git a/Engine/Source/Runtime/CinematicCamera/Public/CineCameraSettings.h b/Engine/Source/Runtime/CinematicCamera/Public/CineCameraSettings.h
--- a/Engine/Source/Runtime/CinematicCamera/Public/CineCameraSettings.h
+++ b/Engine/Source/Runtime/CinematicCamera/Public/CineCameraSettings.h
@@ -33,7 +33,7 @@ struct FCameraFilmbackSettings
 	float SensorVerticalOffset;
 	
 	/** Read-only. Computed from Sensor dimensions. */
-	UPROPERTY(Interp, VisibleAnywhere, BlueprintReadOnly, Category = "Filmback")
+	UPROPERTY(VisibleAnywhere, BlueprintReadOnly, Category = "Filmback")
 	float SensorAspectRatio;
 
 	bool operator==(const FCameraFilmbackSettings& Other) const
